Added horizontal scrolling to the output log.

diff --git a/Engine/Source/Developer/OutputLog/Private/SOutputLog.cpp b/Engine/Source/Developer/OutputLog/Private/SOutputLog.cpp
--- a/Engine/Source/Developer/OutputLog/Private/SOutputLog.cpp
+++ b/Engine/Source/Developer/OutputLog/Private/SOutputLog.cpp
@@ -475,9 +475,14 @@ void SOutputLog::Construct( const FArguments& InArgs )
 				+SHorizontalBox::Slot()
 				.FillWidth(1)
 				[
-					SNew(SScrollBorder, MessageListView.ToSharedRef())
+					SNew( SScrollBox )
+					.Orientation( Orient_Horizontal )
+					+ SScrollBox::Slot()
 					[
-						MessageListView.ToSharedRef()
+						SNew(SScrollBorder, MessageListView.ToSharedRef())
+						[
+							MessageListView.ToSharedRef()
+						]
 					]
 				]
 				+SHorizontalBox::Slot()
