Some tweaks to minimize horizontal scrollbar artifacts in the output log.

diff --git a/Engine/Source/Developer/OutputLog/Private/SOutputLog.cpp b/Engine/Source/Developer/OutputLog/Private/SOutputLog.cpp
--- a/Engine/Source/Developer/OutputLog/Private/SOutputLog.cpp
+++ b/Engine/Source/Developer/OutputLog/Private/SOutputLog.cpp
@@ -463,6 +463,12 @@ void SOutputLog::Construct( const FArguments& InArgs )
 		.OnContextMenuOpening(this, &SOutputLog::BuildMenuWidget)
 		.ExternalScrollbar(OutputLogScrollBar);
 
+	TSharedRef<SScrollBar> HorizontalScrollBar = 
+		SNew( SScrollBar )
+		.Orientation( Orient_Horizontal )
+		.AlwaysShowScrollbar( true )
+		.Thickness( FVector2D( 12.0, 9 ) );
+
 	ChildSlot
 	[
 		SNew(SVerticalBox)
@@ -477,6 +483,7 @@ void SOutputLog::Construct( const FArguments& InArgs )
 				[
 					SNew( SScrollBox )
 					.Orientation( Orient_Horizontal )
+					.ExternalScrollbar( HorizontalScrollBar )
 					+ SScrollBox::Slot()
 					[
 						SNew(SScrollBorder, MessageListView.ToSharedRef())
@@ -491,17 +498,16 @@ void SOutputLog::Construct( const FArguments& InArgs )
 					SNew( SBox )
 					.WidthOverride( FOptionalSize( 16 ) )
 					[
-						SNew(SVerticalBox)						
-						+SVerticalBox::Slot()
-						.FillHeight(1)
-						[
-							// Output log area
-							OutputLogScrollBar->AsShared()
-						]
+						// Output log area
+						OutputLogScrollBar->AsShared()
 					]
 				]
 			]
-
+			+SVerticalBox::Slot()
+			.AutoHeight()
+			[
+				HorizontalScrollBar
+			]
 			// The console input box
 			+SVerticalBox::Slot()
 			.AutoHeight()
