Fixed crash shutting down the editor

TTP# 345402 - CRITICAL: Regression: EDITOR: CRASH: Closing the editor

There was a static instance of SOutputLog being created, which was persisting beyond the lifespan of both Slate and ICU.

diff --git a/Engine/Source/Developer/OutputLog/Private/OutputLogModule.cpp b/Engine/Source/Developer/OutputLog/Private/OutputLogModule.cpp
--- a/Engine/Source/Developer/OutputLog/Private/OutputLogModule.cpp
+++ b/Engine/Source/Developer/OutputLog/Private/OutputLogModule.cpp
@@ -56,20 +56,14 @@ private:
 /** Our global output log app spawner */
 static TSharedPtr<FOutputLogHistory> OutputLogHistory;
 
-
-/** Our global output log app spawner */
-static TSharedPtr<SOutputLog> OutputLogInstance;
-
 TSharedRef<SDockTab> SpawnOutputLog( const FSpawnTabArgs& Args )
 {
-	OutputLogInstance = SNew(SOutputLog).Messages( OutputLogHistory->GetMessages() );
-
 	return SNew(SDockTab)
 		.Icon(FEditorStyle::GetBrush("Log.TabIcon"))
 		.TabRole( ETabRole::NomadTab )
 		.Label( NSLOCTEXT("OutputLog", "TabTitle", "Output Log") )
 		[
-			OutputLogInstance.ToSharedRef()
+			SNew(SOutputLog).Messages( OutputLogHistory->GetMessages() )
 		];
 }
 
