Console can be opened anywhere in the editor by pressing tilde (~).  This is shortcut is rebindable.

diff --git a/Engine/Source/Developer/OutputLog/Private/OutputLogModule.cpp b/Engine/Source/Developer/OutputLog/Private/OutputLogModule.cpp
--- a/Engine/Source/Developer/OutputLog/Private/OutputLogModule.cpp
+++ b/Engine/Source/Developer/OutputLog/Private/OutputLogModule.cpp
@@ -10,7 +10,7 @@ IMPLEMENT_MODULE( FOutputLogModule, OutputLog );
 
 namespace OutputLogModule
 {
-	static const FName OutputLogApp = FName(TEXT("OutputLogApp"));
+	static const FName OutputLogTabName = FName(TEXT("OutputLog"));
 }
 
 /** This class is to capture all log output even if the log window is closed */
@@ -71,7 +71,7 @@ void FOutputLogModule::StartupModule()
 {
 	FOutputLogCommands::Register();
 
-	FGlobalTabmanager::Get()->RegisterNomadTabSpawner("OutputLog", FOnSpawnTab::CreateStatic( &SpawnOutputLog ) )
+	FGlobalTabmanager::Get()->RegisterNomadTabSpawner(OutputLogModule::OutputLogTabName, FOnSpawnTab::CreateStatic( &SpawnOutputLog ) )
 		.SetDisplayName(NSLOCTEXT("UnrealEditor", "OutputLogTab", "Output Log"))
 		.SetTooltipText(NSLOCTEXT("UnrealEditor", "OutputLogTooltipText", "Open the Output Log tab."))
 		.SetGroup( WorkspaceMenu::GetMenuStructure().GetToolsCategory() )
@@ -84,7 +84,7 @@ void FOutputLogModule::ShutdownModule()
 {
 	if (FSlateApplication::IsInitialized())
 	{
-		FGlobalTabmanager::Get()->UnregisterNomadTabSpawner("OutputLog");
+		FGlobalTabmanager::Get()->UnregisterNomadTabSpawner(OutputLogModule::OutputLogTabName);
 	}
 	
 
@@ -99,37 +99,47 @@ TSharedRef< SWidget > FOutputLogModule::MakeConsoleInputBox( TSharedPtr< SEditab
 }
 
 
-void FOutputLogModule::OpenDebugConsoleForWindow( const TSharedRef< SWindow >& Window, const EDebugConsoleStyle::Type InStyle )
+void FOutputLogModule::ToggleDebugConsoleForWindow( const TSharedRef< SWindow >& Window, const EDebugConsoleStyle::Type InStyle, const FDebugConsoleDelegates& DebugConsoleDelegates )
 {
+	bool bShouldOpen = true;
 	// Close an existing console box, if there is one
-	bool bAlreadyOpenForThisWindow = false;
 	TSharedPtr< SWidget > PinnedDebugConsole( DebugConsole.Pin() );
 	if( PinnedDebugConsole.IsValid() )
 	{
-		TSharedPtr< SWindow > WindowForExistingConsole = FSlateApplication::Get().FindWidgetWindow( PinnedDebugConsole.ToSharedRef() );
-		if( WindowForExistingConsole.IsValid() )
+		// If the console is already open close it unless it is in a different window.  In that case reopen it on that window
+		bShouldOpen = false;
+		TSharedPtr< SWindow > WindowForExistingConsole = FSlateApplication::Get().FindWidgetWindow(PinnedDebugConsole.ToSharedRef());
+		if (WindowForExistingConsole.IsValid())
 		{
-			if( WindowForExistingConsole == Window )
-			{
-				bAlreadyOpenForThisWindow = true;
-			}
-			else
-			{
-				CloseDebugConsoleForWindow( WindowForExistingConsole.ToSharedRef() );
-			}
+			WindowForExistingConsole->RemoveOverlaySlot(PinnedDebugConsole.ToSharedRef());
+			DebugConsole.Reset();
 		}
 
-		DebugConsole.Reset();
+		if( WindowForExistingConsole != Window )
+		{
+			// Console is being opened on another window
+			bShouldOpen = true;
+		}
+	}
+	
+	TSharedPtr<SDockTab> ActiveTab = FGlobalTabmanager::Get()->GetActiveTab();
+	if (ActiveTab.IsValid() && ActiveTab->GetLayoutIdentifier() == FTabId(OutputLogModule::OutputLogTabName))
+	{
+		FGlobalTabmanager::Get()->DrawAttention(ActiveTab.ToSharedRef());
+		bShouldOpen = false;
 	}
 
-	if( !bAlreadyOpenForThisWindow )
+	if( bShouldOpen )
 	{
 		const EDebugConsoleStyle::Type DebugConsoleStyle = InStyle;
-		TSharedRef< SDebugConsole > DebugConsoleRef = SNew( SDebugConsole, DebugConsoleStyle, this );
+		TSharedRef< SDebugConsole > DebugConsoleRef = SNew( SDebugConsole, DebugConsoleStyle, this, &DebugConsoleDelegates );
 		DebugConsole = DebugConsoleRef;
 
 		const int32 MaximumZOrder = MAX_int32;
 		Window->AddOverlaySlot( MaximumZOrder )
+		.VAlign(VAlign_Bottom)
+		.HAlign(HAlign_Center)
+		.Padding( 10.0f )
 		[
 			DebugConsoleRef
 		];
@@ -140,14 +150,17 @@ void FOutputLogModule::OpenDebugConsoleForWindow( const TSharedRef< SWindow >& W
 }
 
 
-void FOutputLogModule::CloseDebugConsoleForWindow( const TSharedRef< SWindow >& Window )
+void FOutputLogModule::CloseDebugConsole()
 {
-#if !UE_BUILD_SHIPPING
 	TSharedPtr< SWidget > PinnedDebugConsole( DebugConsole.Pin() );
+
 	if( PinnedDebugConsole.IsValid() )
 	{
-		Window->RemoveOverlaySlot( PinnedDebugConsole.ToSharedRef() );
-		DebugConsole.Reset();
+		TSharedPtr< SWindow > WindowForExistingConsole = FSlateApplication::Get().FindWidgetWindow(PinnedDebugConsole.ToSharedRef());
+		if (WindowForExistingConsole.IsValid())
+		{
+			WindowForExistingConsole->RemoveOverlaySlot( PinnedDebugConsole.ToSharedRef() );
+			DebugConsole.Reset();
+		}
 	}
-#endif
 }
\ No newline at end of file
diff --git a/Engine/Source/Developer/OutputLog/Private/SDebugConsole.cpp b/Engine/Source/Developer/OutputLog/Private/SDebugConsole.cpp
--- a/Engine/Source/Developer/OutputLog/Private/SDebugConsole.cpp
+++ b/Engine/Source/Developer/OutputLog/Private/SDebugConsole.cpp
@@ -12,10 +12,13 @@ namespace DebugConsoleDefs
 }
 
 BEGIN_SLATE_FUNCTION_BUILD_OPTIMIZATION
-void SDebugConsole::Construct( const FArguments& InArgs, const EDebugConsoleStyle::Type InStyle, FOutputLogModule* OutputLogModule )
+
+void SDebugConsole::Construct( const FArguments& InArgs, const EDebugConsoleStyle::Type InStyle, FOutputLogModule* OutputLogModule, const FDebugConsoleDelegates* DebugConsoleDelegates )
 {
 	CurrentStyle = InStyle;
 
+	TSharedPtr<SConsoleInputBox> ConsoleInputBox;
+
 	check( OutputLogModule != NULL );
 	ChildSlot
 	[ 
@@ -34,7 +37,7 @@ void SDebugConsole::Construct( const FArguments& InArgs, const EDebugConsoleStyl
 					.HeightOverride( 200.0f )
 					[
 						SNew( SBorder )
-							.BorderImage( FEditorStyle::GetBrush( "DebugConsole.Background" ) )
+							.BorderImage( FEditorStyle::GetBrush( "ToolPanel.GroupBorder" ) )
 							.ColorAndOpacity( this, &SDebugConsole::GetAnimatedColorAndOpacity )
 							.BorderBackgroundColor( this, &SDebugConsole::GetAnimatedSlateColor )
 							[
@@ -50,35 +53,45 @@ void SDebugConsole::Construct( const FArguments& InArgs, const EDebugConsoleStyl
 		[
 			SNew(SBox)
 			.HeightOverride( 26.0f )
+			.HAlign( HAlign_Left )
 			[
 				SNew( SBorder )
-					.BorderImage( FEditorStyle::GetBrush( "Console.Background" ) )
-					.ColorAndOpacity( this, &SDebugConsole::GetAnimatedColorAndOpacity )
-					.BorderBackgroundColor( this, &SDebugConsole::GetAnimatedSlateColor )
+				.Padding( FMargin(2) )
+				.BorderImage( FEditorStyle::GetBrush( "DebugConsole.Background" ) )
+				.ColorAndOpacity( this, &SDebugConsole::GetAnimatedColorAndOpacity )
+				.BorderBackgroundColor( this, &SDebugConsole::GetAnimatedSlateColor )
+				[
+					SNew(SHorizontalBox)
+					+ SHorizontalBox::Slot()
+					.AutoWidth()
+					.Padding(3.0f)
+					.VAlign(VAlign_Center)
 					[
-						SNew( SHorizontalBox )
-							+SHorizontalBox::Slot()
-								.AutoWidth()
-								.Padding( 3.0f )
-							[
-								SNew( STextBlock )
-									.Text( NSLOCTEXT("Console", "ConsoleLabel", "Console >").ToString() )
-							]
-							+SHorizontalBox::Slot()
-								.Padding( 3.0f )
-								.FillWidth( 1.0f )
-							[
-								OutputLogModule->MakeConsoleInputBox( EditableTextBox )
-							]
+						SNew(STextBlock)
+						.Text(NSLOCTEXT("Console", "ConsoleLabel", "Console"))
+
 					]
+					+ SHorizontalBox::Slot()
+					.Padding(5.0f, 2.0f)
+					.VAlign(VAlign_Center)
+					.MaxWidth(400.0f)
+					[
+						SAssignNew(ConsoleInputBox, SConsoleInputBox)
+						.OnConsoleCommandExecuted(DebugConsoleDelegates->OnConsoleCommandExecuted)
+					]
+				]
 			]
 		]
 	];
 
+	EditableTextBox = ConsoleInputBox->GetEditableTextBox();
+
 	// Kick off intro animation
-	AnimCurve = FCurveSequence();
-	AnimCurve.AddCurve( 0.0f, DebugConsoleDefs::IntroAnimationDuration, ECurveEaseFunction::QuadOut );
-	AnimCurve.Play();
+	AnimCurveSequence = FCurveSequence();
+	AnimCurve = AnimCurveSequence.AddCurve( 0.0f, DebugConsoleDefs::IntroAnimationDuration, ECurveEaseFunction::QuadOut );
+	FlashCurve = AnimCurveSequence.AddCurve( DebugConsoleDefs::IntroAnimationDuration, .15f, ECurveEaseFunction::QuadInOut );
+
+	AnimCurveSequence.Play();
 }
 END_SLATE_FUNCTION_BUILD_OPTIMIZATION
 
@@ -90,15 +103,9 @@ SDebugConsole::SDebugConsole()
 
 void SDebugConsole::SetFocusToEditableText()
 {
-	FWidgetPath WidgetToFocusPath;
-	FSlateApplication::Get().GeneratePathToWidgetChecked( EditableTextBox.ToSharedRef(), WidgetToFocusPath );
-	if( WidgetToFocusPath.IsValid() )
-	{
-		FSlateApplication::Get().SetKeyboardFocus( WidgetToFocusPath, EKeyboardFocusCause::SetDirectly );
-	}
+	FSlateApplication::Get().SetKeyboardFocus( EditableTextBox.ToSharedRef(), EKeyboardFocusCause::SetDirectly );
 }
 
-
 EVisibility SDebugConsole::MakeVisibleIfLogIsShown() const
 {
 	return CurrentStyle == EDebugConsoleStyle::WithLog ? EVisibility::Visible : EVisibility::Collapsed;
@@ -115,3 +122,15 @@ FSlateColor SDebugConsole::GetAnimatedSlateColor() const
 {
 	return FSlateColor( GetAnimatedColorAndOpacity() );
 }
+
+FSlateColor SDebugConsole::GetFlashColor() const
+{
+	float FlashAlpha = 1.0f - FlashCurve.GetLerp();
+
+	if (FlashAlpha == 1.0f)
+	{
+		FlashAlpha = 0.0f;
+	}
+
+	return FLinearColor(1,1,1,FlashAlpha);
+}
\ No newline at end of file
diff --git a/Engine/Source/Developer/OutputLog/Private/SDebugConsole.h b/Engine/Source/Developer/OutputLog/Private/SDebugConsole.h
--- a/Engine/Source/Developer/OutputLog/Private/SDebugConsole.h
+++ b/Engine/Source/Developer/OutputLog/Private/SDebugConsole.h
@@ -17,7 +17,7 @@ public:
 	SLATE_END_ARGS()
 
 	/** Constructs this widget */
-	void Construct( const FArguments& InArgs, const EDebugConsoleStyle::Type InStyle, FOutputLogModule* OutputLogModule );
+	void Construct( const FArguments& InArgs, const EDebugConsoleStyle::Type InStyle, FOutputLogModule* OutputLogModule, const FDebugConsoleDelegates* DebugConsoleDelegates );
 
 	/** Call to set focus to this debug console's editable text box */
 	void SetFocusToEditableText();
@@ -27,7 +27,6 @@ public:
 
 
 protected:
-
 	/** Returns EVisibility::Visible if style has log being shown, otherwise VIS_COLLAPSED */
 	EVisibility MakeVisibleIfLogIsShown() const;
 
@@ -37,7 +36,7 @@ protected:
 	/** Returns a Slate color based on any current animation (same color as GetAnimatedColorAndOpacity) */
 	FSlateColor GetAnimatedSlateColor() const;
 
-
+	FSlateColor GetFlashColor() const;
 private:
 
 	/** Editable text box for this debug console's input line */
@@ -47,6 +46,8 @@ private:
 	EDebugConsoleStyle::Type CurrentStyle;
 
 	/** Intro/outro animation curve */
-	FCurveSequence AnimCurve;
+	FCurveSequence AnimCurveSequence;
+	FCurveHandle AnimCurve;
+	FCurveHandle FlashCurve;
 };
 
diff --git a/Engine/Source/Developer/OutputLog/Private/SOutputLog.cpp b/Engine/Source/Developer/OutputLog/Private/SOutputLog.cpp
--- a/Engine/Source/Developer/OutputLog/Private/SOutputLog.cpp
+++ b/Engine/Source/Developer/OutputLog/Private/SOutputLog.cpp
@@ -5,6 +5,119 @@
 #include "OutputLogActions.h"
 #include "SScrollBorder.h"
 
+/** Custom console editable text box whose only purpose is to prevent some keys from being typed */
+class SConsoleEditableTextBox : public SEditableTextBox
+{
+public:
+	SLATE_BEGIN_ARGS( SConsoleEditableTextBox ) {}
+		
+		/** Hint text that appears when there is no text in the text box */
+		SLATE_ATTRIBUTE(FText, HintText)
+
+		/** Called whenever the text is changed interactively by the user */
+		SLATE_EVENT(FOnTextChanged, OnTextChanged)
+
+		/** Called whenever the text is committed.  This happens when the user presses enter or the text box loses focus. */
+		SLATE_EVENT(FOnTextCommitted, OnTextCommitted)
+
+	SLATE_END_ARGS()
+
+
+	void Construct( const FArguments& InArgs )
+	{
+		SetStyle(&FCoreStyle::Get().GetWidgetStyle< FEditableTextBoxStyle >("NormalEditableTextBox"));
+
+		SBorder::Construct(SBorder::FArguments()
+			.BorderImage(this, &SConsoleEditableTextBox::GetConsoleBorder)
+			.BorderBackgroundColor(Style->BackgroundColor)
+			.ForegroundColor(Style->ForegroundColor)
+			.Padding(Style->Padding)
+			[
+				SAssignNew( EditableText, SConsoleEditableText )
+				.HintText( InArgs._HintText )
+				.OnTextChanged( InArgs._OnTextChanged )
+				.OnTextCommitted( InArgs._OnTextCommitted )
+			] );
+	}
+
+private:
+	class SConsoleEditableText : public SEditableText
+	{
+	public:
+		SLATE_BEGIN_ARGS( SConsoleEditableText ) {}
+			/** The text that appears when there is nothing typed into the search box */
+			SLATE_ATTRIBUTE(FText, HintText)
+			/** Called whenever the text is changed interactively by the user */
+			SLATE_EVENT(FOnTextChanged, OnTextChanged)
+
+			/** Called whenever the text is committed.  This happens when the user presses enter or the text box loses focus. */
+			SLATE_EVENT(FOnTextCommitted, OnTextCommitted)
+		SLATE_END_ARGS()
+
+		void Construct( const FArguments& InArgs )
+		{
+			SEditableText::Construct
+			( 
+				SEditableText::FArguments()
+				.HintText( InArgs._HintText )
+				.OnTextChanged( InArgs._OnTextChanged )
+				.OnTextCommitted( InArgs._OnTextCommitted ) 
+				.ClearKeyboardFocusOnCommit( false )
+				.IsCaretMovedWhenGainFocus( false ) 
+				.MinDesiredWidth( 400.0f )
+			);
+		}
+
+		virtual FReply OnKeyDown( const FGeometry& MyGeometry, const FKeyboardEvent& InKeyboardEvent )
+		{
+			// Special case handling.  Intercept the tilde key.  It is not suitable for typing in the console
+			if( InKeyboardEvent.GetKey() == EKeys::Tilde )
+			{
+				return FReply::Unhandled();
+			}
+			else
+			{
+				return SEditableText::OnKeyDown( MyGeometry, InKeyboardEvent );
+			}
+		}
+
+		virtual FReply OnKeyChar( const FGeometry& MyGeometry, const FCharacterEvent& InCharacterEvent )
+		{
+			// Special case handling.  Intercept the tilde key.  It is not suitable for typing in the console
+			if( InCharacterEvent.GetCharacter() != 0x60 )
+			{
+				return SEditableText::OnKeyChar( MyGeometry, InCharacterEvent );
+			}
+			else
+			{
+				return FReply::Unhandled();
+			}
+		}
+
+	};
+
+	/** @return Border image for the text box based on the hovered and focused state */
+	const FSlateBrush* GetConsoleBorder() const
+	{
+		if (EditableText->HasKeyboardFocus())
+		{
+			return &Style->BackgroundImageFocused;
+		}
+		else
+		{
+			if (EditableText->IsHovered())
+			{
+				return &Style->BackgroundImageHovered;
+			}
+			else
+			{
+				return &Style->BackgroundImageNormal;
+			}
+		}
+	}
+
+};
+
 SConsoleInputBox::SConsoleInputBox()
 	: SelectedSuggestion(-1)
 	, bIgnoreUIUpdate(false)
@@ -14,16 +127,16 @@ SConsoleInputBox::SConsoleInputBox()
 BEGIN_SLATE_FUNCTION_BUILD_OPTIMIZATION
 void SConsoleInputBox::Construct( const FArguments& InArgs )
 {
+	OnConsoleCommandExecuted = InArgs._OnConsoleCommandExecuted;
+
 	ChildSlot
 	[
 		SAssignNew( SuggestionBox, SMenuAnchor )
 			.Placement( InArgs._SuggestionListPlacement )
 			[
-				SAssignNew(InputText, SEditableTextBox)
+				SAssignNew(InputText, SConsoleEditableTextBox)
 					.OnTextCommitted(this, &SConsoleInputBox::OnTextCommitted)
 					.HintText( NSLOCTEXT( "ConsoleInputBox", "TypeInConsoleHint", "Enter console command" ) )
-					.ClearKeyboardFocusOnCommit(false)
-					.IsCaretMovedWhenGainFocus(false)
 					.OnTextChanged(this, &SConsoleInputBox::OnTextChanged)
 			]
 			.MenuContent
@@ -266,6 +379,8 @@ void SConsoleInputBox::OnTextCommitted( const FText& InText, ETextCommit::Type C
 		}
 
 		ClearSuggestions();
+
+		OnConsoleCommandExecuted.ExecuteIfBound();
 	}
 }
 
diff --git a/Engine/Source/Developer/OutputLog/Private/SOutputLog.h b/Engine/Source/Developer/OutputLog/Private/SOutputLog.h
--- a/Engine/Source/Developer/OutputLog/Private/SOutputLog.h
+++ b/Engine/Source/Developer/OutputLog/Private/SOutputLog.h
@@ -36,6 +36,8 @@ public:
 		/** Where to place the suggestion list */
 		SLATE_ARGUMENT( EMenuPlacement, SuggestionListPlacement )
 
+		/** Called when a console command is executed */
+		SLATE_EVENT( FSimpleDelegate, OnConsoleCommandExecuted )
 	SLATE_END_ARGS()
 
 	/** Protected console input box widget constructor, called by Slate */
@@ -99,6 +101,9 @@ private:
 	/** The list view for showing all log messages. Should be replaced by a full text editor */
 	TSharedPtr< SListView< TSharedPtr<FString> > > SuggestionListView;
 
+	/** Delegate to call when a console command is executed */
+	FSimpleDelegate OnConsoleCommandExecuted;
+
 	/** -1 if not set, otherwise index into Suggestions */
 	int32 SelectedSuggestion;
 
diff --git a/Engine/Source/Developer/OutputLog/Public/OutputLogModule.h b/Engine/Source/Developer/OutputLog/Public/OutputLogModule.h
--- a/Engine/Source/Developer/OutputLog/Public/OutputLogModule.h
+++ b/Engine/Source/Developer/OutputLog/Public/OutputLogModule.h
@@ -17,6 +17,11 @@ namespace EDebugConsoleStyle
 	};
 };
 
+struct FDebugConsoleDelegates
+{
+	FSimpleDelegate OnKeyboardFocusLost;
+	FSimpleDelegate OnConsoleCommandExecuted;
+};
 
 class FOutputLogModule : public IModuleInterface
 {
@@ -29,10 +34,10 @@ public:
 	virtual TSharedRef< SWidget > MakeConsoleInputBox( TSharedPtr< SEditableTextBox >& OutExposedEditableTextBox ) const;
 
 	/** Opens a debug console in the specified window, if not already open */
-	virtual void OpenDebugConsoleForWindow( const TSharedRef< SWindow >& Window, const EDebugConsoleStyle::Type InStyle );
+	virtual void ToggleDebugConsoleForWindow( const TSharedRef< SWindow >& Window, const EDebugConsoleStyle::Type InStyle, const FDebugConsoleDelegates& DebugConsoleDelegates );
 
 	/** Closes the debug console for the specified window */
-	virtual void CloseDebugConsoleForWindow( const TSharedRef< SWindow >& Window );
+	virtual void CloseDebugConsole();
 
 
 private:
diff --git a/Engine/Source/Editor/EditorStyle/Private/SlateEditorStyle.cpp b/Engine/Source/Editor/EditorStyle/Private/SlateEditorStyle.cpp
--- a/Engine/Source/Editor/EditorStyle/Private/SlateEditorStyle.cpp
+++ b/Engine/Source/Editor/EditorStyle/Private/SlateEditorStyle.cpp
@@ -883,6 +883,8 @@ void FSlateEditorStyle::FStyle::SetupGeneralStyles()
 			.SetFont( TTF_FONT( "Fonts/DroidSansMono", 9 ) ) );
 
 		Set("Log.TabIcon", new IMAGE_BRUSH( "Icons/icon_tab_OutputLog_16x", Icon16x16 ) );
+
+		Set("DebugConsole.Background", new BOX_BRUSH("Old/Menu_Background", FMargin(8.0f / 64.0f)));
 	}
 
 	// Debug Tools Window
diff --git a/Engine/Source/Editor/UnrealEd/Private/EditorViewportCommands.cpp b/Engine/Source/Editor/UnrealEd/Private/EditorViewportCommands.cpp
--- a/Engine/Source/Editor/UnrealEd/Private/EditorViewportCommands.cpp
+++ b/Engine/Source/Editor/UnrealEd/Private/EditorViewportCommands.cpp
@@ -51,7 +51,7 @@ void FEditorViewportCommands::RegisterCommands()
 	UI_COMMAND( RelativeCoordinateSystem_World, "World-relative Transform", "Move and rotate objects relative to the cardinal world axes", EUserInterfaceActionType::RadioButton, FInputGesture() );
 	UI_COMMAND( RelativeCoordinateSystem_Local, "Local-relative Transform", "Move and rotate objects relative to the object's local axes", EUserInterfaceActionType::RadioButton, FInputGesture() );
 
-	UI_COMMAND( CycleTransformGizmoCoordSystem, "Cycle Transform Coordinate System", "Cycles the transform gizmo coordinate systems between world and local (object) space", EUserInterfaceActionType::Button, FInputGesture(EKeys::Tilde));
+	UI_COMMAND( CycleTransformGizmoCoordSystem, "Cycle Transform Coordinate System", "Cycles the transform gizmo coordinate systems between world and local (object) space", EUserInterfaceActionType::Button, FInputGesture(EModifierKey::Control, EKeys::Tilde));
 	UI_COMMAND( CycleTransformGizmos, "Cycle Between Translate, Rotate, and Scale", "Cycles the transform gizmos between translate, rotate, and scale", EUserInterfaceActionType::Button, FInputGesture(EKeys::SpaceBar) );
 	
 	UI_COMMAND( FocusViewportToSelection, "Focus Selected", "Moves the camera in front of the selection", EUserInterfaceActionType::Button, FInputGesture( EKeys::F ) );
diff --git a/Engine/Source/Editor/UnrealEd/Private/Toolkits/GlobalEditorCommonCommands.cpp b/Engine/Source/Editor/UnrealEd/Private/Toolkits/GlobalEditorCommonCommands.cpp
--- a/Engine/Source/Editor/UnrealEd/Private/Toolkits/GlobalEditorCommonCommands.cpp
+++ b/Engine/Source/Editor/UnrealEd/Private/Toolkits/GlobalEditorCommonCommands.cpp
@@ -9,6 +9,7 @@
 #include "AssetData.h"
 #include "ContentBrowserModule.h"
 #include "Toolkits/AssetEditorManager.h"
+#include "OutputLogModule.h"
 
 class SGlobalOpenAssetDialog : public SCompoundWidget
 {
@@ -98,6 +99,8 @@ void FGlobalEditorCommonCommands::RegisterCommands()
 	UI_COMMAND( FindInContentBrowser, "Find in Content Browser", "Summons the Content Browser and navigates to the selected asset", EUserInterfaceActionType::Button, FInputGesture(EModifierKey::Control, EKeys::B));
 	UI_COMMAND( ViewReferences, "View References", "Launches the reference viewer showing the selected assets' references", EUserInterfaceActionType::Button, FInputGesture(EModifierKey::Shift | EModifierKey::Alt, EKeys::R));
 	
+	UI_COMMAND( OpenConsoleCommandBox, "Open Console Command Box", "Opens an edit box where you can type in a console command", EUserInterfaceActionType::Button, FInputGesture(EKeys::Tilde));
+
 	UI_COMMAND( OpenDocumentation, "Open Documentation...", "Opens documentation for this tool", EUserInterfaceActionType::Button, FInputGesture(EKeys::F1) );
 }
 
@@ -112,6 +115,10 @@ void FGlobalEditorCommonCommands::MapActions(TSharedRef<FUICommandList>& Toolkit
 	ToolkitCommands->MapAction(
 		Get().SummonOpenAssetDialog,
 		FExecuteAction::CreateStatic( &FGlobalEditorCommonCommands::OnSummonedAssetPicker ) );
+
+	ToolkitCommands->MapAction(
+		Get().OpenConsoleCommandBox,
+		FExecuteAction::CreateStatic( &FGlobalEditorCommonCommands::OnSummonedConsoleCommandBox ) );
 }
 
 void FGlobalEditorCommonCommands::OnPressedCtrlTab()
@@ -150,3 +157,25 @@ void FGlobalEditorCommonCommands::OnSummonedAssetPicker()
 		TSharedRef<SWindow> PopupWindow = FSlateApplication::Get().PushMenu(ParentWindow.ToSharedRef(), WindowContents, WindowPosition, TransitionEffect);
 	}
 }
+
+static void CloseDebugConsole()
+{
+	FOutputLogModule& OutputLogModule = FModuleManager::LoadModuleChecked< FOutputLogModule >(TEXT("OutputLog"));
+	OutputLogModule.CloseDebugConsole();
+}
+
+void FGlobalEditorCommonCommands::OnSummonedConsoleCommandBox()
+{
+	TSharedPtr<SWindow> ParentWindow = FSlateApplication::Get().GetActiveTopLevelWindow();
+	
+	if( ParentWindow.IsValid() )
+	{
+		TSharedRef<SWindow> WindowRef = ParentWindow.ToSharedRef();
+		FOutputLogModule& OutputLogModule = FModuleManager::LoadModuleChecked< FOutputLogModule >(TEXT("OutputLog"));
+
+		FDebugConsoleDelegates Delegates;
+		Delegates.OnConsoleCommandExecuted = FSimpleDelegate::CreateStatic( &CloseDebugConsole );
+
+		OutputLogModule.ToggleDebugConsoleForWindow(WindowRef, EDebugConsoleStyle::Compact, Delegates);
+	}
+}
\ No newline at end of file
diff --git a/Engine/Source/Editor/UnrealEd/Public/Toolkits/GlobalEditorCommonCommands.h b/Engine/Source/Editor/UnrealEd/Public/Toolkits/GlobalEditorCommonCommands.h
--- a/Engine/Source/Editor/UnrealEd/Public/Toolkits/GlobalEditorCommonCommands.h
+++ b/Engine/Source/Editor/UnrealEd/Public/Toolkits/GlobalEditorCommonCommands.h
@@ -17,11 +17,13 @@ public:
 protected:
 	static void OnPressedCtrlTab();
 	static void OnSummonedAssetPicker();
+	static void OnSummonedConsoleCommandBox();
 public:
 	TSharedPtr<FUICommandInfo> FindInContentBrowser;
 	TSharedPtr<FUICommandInfo> SummonControlTabNavigation;
 	TSharedPtr<FUICommandInfo> SummonOpenAssetDialog;
 	TSharedPtr<FUICommandInfo> OpenDocumentation;
 	TSharedPtr<FUICommandInfo> ViewReferences;
+	TSharedPtr<FUICommandInfo> OpenConsoleCommandBox;
 };
 
diff --git a/Engine/Source/Editor/UnrealEd/UnrealEd.Build.cs b/Engine/Source/Editor/UnrealEd/UnrealEd.Build.cs
--- a/Engine/Source/Editor/UnrealEd/UnrealEd.Build.cs
+++ b/Engine/Source/Editor/UnrealEd/UnrealEd.Build.cs
@@ -48,6 +48,7 @@ public class UnrealEd : ModuleRules
 				"SourceCodeAccess",
 				"ReferenceViewer",
                 "IntroTutorials",
+				"OutputLog"
 			}
 		);
 
