Fixed missing header include.

diff --git a/Engine/Source/Developer/OutputLog/Private/SDebugConsole.cpp b/Engine/Source/Developer/OutputLog/Private/SDebugConsole.cpp
--- a/Engine/Source/Developer/OutputLog/Private/SDebugConsole.cpp
+++ b/Engine/Source/Developer/OutputLog/Private/SDebugConsole.cpp
@@ -3,7 +3,7 @@
 #include "OutputLogPrivatePCH.h"
 #include "OutputLogModule.h"
 #include "SDebugConsole.h"
-
+#include "SOutputLog.h"
 
 namespace DebugConsoleDefs
 {
