Converted SOutputLog to use a multiline editable text control to show its log output

This is using a custom text marshaller to efficiently convert an FLogMessage into something understood by the FTextLayout.

diff --git a/Engine/Source/Developer/OutputLog/Private/OutputLogActions.cpp b/Engine/Source/Developer/OutputLog/Private/OutputLogActions.cpp
deleted file mode 100644
diff --git a/Engine/Source/Developer/OutputLog/Private/OutputLogActions.h b/Engine/Source/Developer/OutputLog/Private/OutputLogActions.h
deleted file mode 100644
diff --git a/Engine/Source/Developer/OutputLog/Private/OutputLogModule.cpp b/Engine/Source/Developer/OutputLog/Private/OutputLogModule.cpp
--- a/Engine/Source/Developer/OutputLog/Private/OutputLogModule.cpp
+++ b/Engine/Source/Developer/OutputLog/Private/OutputLogModule.cpp
@@ -3,7 +3,6 @@
 #include "OutputLogPrivatePCH.h"
 #include "SDebugConsole.h"
 #include "SOutputLog.h"
-#include "OutputLogActions.h"
 #include "Editor/WorkspaceMenuStructure/Public/WorkspaceMenuStructureModule.h"
 
 IMPLEMENT_MODULE( FOutputLogModule, OutputLog );
@@ -69,8 +68,6 @@ TSharedRef<SDockTab> SpawnOutputLog( const FSpawnTabArgs& Args )
 
 void FOutputLogModule::StartupModule()
 {
-	FOutputLogCommands::Register();
-
 	FGlobalTabmanager::Get()->RegisterNomadTabSpawner(OutputLogModule::OutputLogTabName, FOnSpawnTab::CreateStatic( &SpawnOutputLog ) )
 		.SetDisplayName(NSLOCTEXT("UnrealEditor", "OutputLogTab", "Output Log"))
 		.SetTooltipText(NSLOCTEXT("UnrealEditor", "OutputLogTooltipText", "Open the Output Log tab."))
@@ -86,9 +83,6 @@ void FOutputLogModule::ShutdownModule()
 	{
 		FGlobalTabmanager::Get()->UnregisterNomadTabSpawner(OutputLogModule::OutputLogTabName);
 	}
-	
-
-	FOutputLogCommands::Unregister();
 }
 
 TSharedRef< SWidget > FOutputLogModule::MakeConsoleInputBox( TSharedPtr< SEditableTextBox >& OutExposedEditableTextBox ) const
diff --git a/Engine/Source/Developer/OutputLog/Private/SOutputLog.cpp b/Engine/Source/Developer/OutputLog/Private/SOutputLog.cpp
--- a/Engine/Source/Developer/OutputLog/Private/SOutputLog.cpp
+++ b/Engine/Source/Developer/OutputLog/Private/SOutputLog.cpp
@@ -2,8 +2,8 @@
 
 #include "OutputLogPrivatePCH.h"
 #include "SOutputLog.h"
-#include "OutputLogActions.h"
 #include "SScrollBorder.h"
+#include "BaseTextLayoutMarshaller.h"
 
 /** Custom console editable text box whose only purpose is to prevent some keys from being typed */
 class SConsoleEditableTextBox : public SEditableTextBox
@@ -541,48 +541,124 @@ FString SConsoleInputBox::GetSelectionText() const
 	return ret;
 }
 
-BEGIN_SLATE_FUNCTION_BUILD_OPTIMIZATION
-void SOutputLog::Construct( const FArguments& InArgs )
+/** Output log text marshaller to convert an array of FLogMessages into styled lines to be consumed by an FTextLayout */
+class FOutputLogTextLayoutMarshaller : public FBaseTextLayoutMarshaller
 {
-	Messages = InArgs._Messages;
-	OutputLogActions = MakeShareable( new FUICommandList );
-	OutputLogScrollBar = SNew( SScrollBar );
+public:
 
-	const FOutputLogCommandsImpl& Commands = FOutputLogCommands::Get();
+	static TSharedRef< FOutputLogTextLayoutMarshaller > Create(TArray< TSharedPtr<FLogMessage> > InMessages);
 
-	OutputLogActions->MapAction(
-		Commands.CopyOutputLog,
-		FExecuteAction::CreateRaw( this, &SOutputLog::OnCopy ),
-		FCanExecuteAction::CreateSP( this, &SOutputLog::CanCopy ));
+	virtual ~FOutputLogTextLayoutMarshaller();
+	
+	// ITextLayoutMarshaller
+	virtual void SetText(const FString& SourceString, FTextLayout& TargetTextLayout) override;
+	virtual void GetText(FString& TargetString, const FTextLayout& SourceTextLayout) override;
 
-	OutputLogActions->MapAction(
-		Commands.SelectAllInOutputLog,
-		FExecuteAction::CreateRaw( this, &SOutputLog::OnSelectAll ),
-		FCanExecuteAction::CreateSP( this, &SOutputLog::CanSelectAll ));
+	bool AppendMessage(const TCHAR* InText, const ELogVerbosity::Type InVerbosity, const FName& InCategory);
+	void ClearMessages();
+	int32 GetNumMessages() const;
 
-	OutputLogActions->MapAction(
-		Commands.SelectNoneInOutputLog,
-		FExecuteAction::CreateRaw( this, &SOutputLog::OnSelectNone ),
-		FCanExecuteAction::CreateSP( this, &SOutputLog::CanSelectNone ));
+protected:
 
-	OutputLogActions->MapAction(
-		Commands.ClearOutputLog,
-		FExecuteAction::CreateRaw( this, &SOutputLog::OnClearLog ),
-		FCanExecuteAction::CreateSP( this, &SOutputLog::CanClearLog ));
-
-	MessageListView = SNew(SListView< TSharedPtr<FLogMessage> >)	// Ideally we start appending the items at the bottom, not at the top
-		.ListItemsSource(&Messages)
-		.OnGenerateRow(this, &SOutputLog::MakeLogListItemWidget)
-		.SelectionMode(ESelectionMode::Multi)
-		.ItemHeight(14)
-		.OnContextMenuOpening(this, &SOutputLog::BuildMenuWidget)
-		.ExternalScrollbar(OutputLogScrollBar);
-
-	TSharedRef<SScrollBar> HorizontalScrollBar = 
-		SNew( SScrollBar )
-		.Orientation( Orient_Horizontal )
-		.AlwaysShowScrollbar( true )
-		.Thickness( FVector2D( 12.0, 9 ) );
+	FOutputLogTextLayoutMarshaller(TArray< TSharedPtr<FLogMessage> > InMessages);
+
+	void AppendMessageToTextLayout(const TSharedPtr<FLogMessage>& Message);
+
+	/** All log messages to show in the text box */
+	TArray< TSharedPtr<FLogMessage> > Messages;
+
+	FTextLayout* TextLayout;
+};
+
+TSharedRef< FOutputLogTextLayoutMarshaller > FOutputLogTextLayoutMarshaller::Create(TArray< TSharedPtr<FLogMessage> > InMessages)
+{
+	return MakeShareable(new FOutputLogTextLayoutMarshaller(MoveTemp(InMessages)));
+}
+
+FOutputLogTextLayoutMarshaller::~FOutputLogTextLayoutMarshaller()
+{
+}
+
+void FOutputLogTextLayoutMarshaller::SetText(const FString& SourceString, FTextLayout& TargetTextLayout)
+{
+	TextLayout = &TargetTextLayout;
+
+	for(const auto& Message : Messages)
+	{
+		AppendMessageToTextLayout(Message);
+	}
+}
+
+void FOutputLogTextLayoutMarshaller::GetText(FString& TargetString, const FTextLayout& SourceTextLayout)
+{
+}
+
+bool FOutputLogTextLayoutMarshaller::AppendMessage(const TCHAR* InText, const ELogVerbosity::Type InVerbosity, const FName& InCategory)
+{
+	TArray< TSharedPtr<FLogMessage> > NewMessages;
+	if(SOutputLog::CreateLogMessages(InText, InVerbosity, InCategory, NewMessages))
+	{
+		Messages.Append(NewMessages);
+
+		if(TextLayout)
+		{
+			// If we've already been given a text layout, then append these new messages rather than force a refresh of the entire document
+			for(const auto& Message : NewMessages)
+			{
+				AppendMessageToTextLayout(Message);
+			}
+		}
+		else
+		{
+			MakeDirty();
+		}
+
+		return true;
+	}
+
+	return false;
+}
+
+void FOutputLogTextLayoutMarshaller::AppendMessageToTextLayout(const TSharedPtr<FLogMessage>& Message)
+{
+	const FTextBlockStyle& MessageTextStyle = FEditorStyle::Get().GetWidgetStyle<FTextBlockStyle>(Message->Style);
+
+	TSharedRef<FString> LineText = Message->Message;
+
+	TArray<TSharedRef<IRun>> Runs;
+	Runs.Add(FSlateTextRun::Create(FRunInfo(), LineText, MessageTextStyle));
+
+	TextLayout->AddLine(LineText, Runs);
+}
+
+void FOutputLogTextLayoutMarshaller::ClearMessages()
+{
+	Messages.Empty();
+	MakeDirty();
+}
+
+int32 FOutputLogTextLayoutMarshaller::GetNumMessages() const
+{
+	return Messages.Num();
+}
+
+FOutputLogTextLayoutMarshaller::FOutputLogTextLayoutMarshaller(TArray< TSharedPtr<FLogMessage> > InMessages)
+	: Messages(MoveTemp(InMessages))
+	, TextLayout(nullptr)
+{
+}
+
+BEGIN_SLATE_FUNCTION_BUILD_OPTIMIZATION
+void SOutputLog::Construct( const FArguments& InArgs )
+{
+	MessagesTextMarshaller = FOutputLogTextLayoutMarshaller::Create(MoveTemp(InArgs._Messages));
+
+	MessagesTextBox = SNew(SMultiLineEditableTextBox)
+		.Style(FEditorStyle::Get(), "Log.TextBox")
+		.TextStyle(FEditorStyle::Get(), "Log.Normal")
+		.Marshaller(MessagesTextMarshaller)
+		.IsReadOnly(true)
+		.ContextMenuExtender(this, &SOutputLog::ExtendTextBoxMenu);
 
 	ChildSlot
 	[
@@ -592,40 +668,12 @@ void SOutputLog::Construct( const FArguments& InArgs )
 			+SVerticalBox::Slot()
 			.FillHeight(1)
 			[
-				SNew( SHorizontalBox )
-				+SHorizontalBox::Slot()
-				.FillWidth(1)
-				[
-					SNew( SScrollBox )
-					.Orientation( Orient_Horizontal )
-					.ExternalScrollbar( HorizontalScrollBar )
-					+ SScrollBox::Slot()
-					[
-						SNew(SScrollBorder, MessageListView.ToSharedRef())
-						[
-							MessageListView.ToSharedRef()
-						]
-					]
-				]
-				+SHorizontalBox::Slot()
-				.AutoWidth()
-				[
-					SNew( SBox )
-					.WidthOverride( FOptionalSize( 16 ) )
-					[
-						// Output log area
-						OutputLogScrollBar->AsShared()
-					]
-				]
-			]
-			+SVerticalBox::Slot()
-			.AutoHeight()
-			[
-				HorizontalScrollBar
+				MessagesTextBox.ToSharedRef()
 			]
 			// The console input box
 			+SVerticalBox::Slot()
 			.AutoHeight()
+			.Padding(FMargin(0.0f, 4.0f, 0.0f, 0.0f))
 			[
 				SNew( SConsoleInputBox )
 
@@ -637,9 +685,9 @@ void SOutputLog::Construct( const FArguments& InArgs )
 	GLog->AddOutputDevice(this);
 
 	// If there's already been messages logged, scroll down to the last one
-	if (Messages.Num() > 0)
+	if (MessagesTextMarshaller->GetNumMessages() > 0)
 	{
-		MessageListView->RequestScrollIntoView(Messages.Last());
+		MessagesTextBox->ScrollTo(FTextLocation(MessagesTextMarshaller->GetNumMessages() - 1));
 	}
 }
 END_SLATE_FUNCTION_BUILD_OPTIMIZATION
@@ -664,178 +712,86 @@ bool SOutputLog::CreateLogMessages( const TCHAR* V, ELogVerbosity::Type Verbosit
 		FName Style;
 		if (Category == NAME_Cmd)
 		{
-			Style = FName(TEXT("LogTableRow.Command"));
+			Style = FName(TEXT("Log.Command"));
 		}
 		else if (Verbosity == ELogVerbosity::Error)
 		{
-			Style = FName(TEXT("LogTableRow.Error"));
+			Style = FName(TEXT("Log.Error"));
 		}
 		else if (Verbosity == ELogVerbosity::Warning)
 		{
-			Style = FName(TEXT("LogTableRow.Warning"));
+			Style = FName(TEXT("Log.Warning"));
 		}
 		else
 		{
-			Style = FName(TEXT("LogTableRow.Normal"));
+			Style = FName(TEXT("Log.Normal"));
 		}
 
+		const int32 OldNumMessages = OutMessages.Num();
+
 		// handle multiline strings by breaking them apart by line
-		TArray< FString > MessageLines;
+		TArray<FTextRange> LineRanges;
 		FString CurrentLogDump = V;
-		CurrentLogDump.ParseIntoArray(&MessageLines, TEXT("\n"), false);
+		FTextRange::CalculateLineRangesFromString(CurrentLogDump, LineRanges);
 
-		for (int32 i = 0; i < MessageLines.Num(); ++i)
+		bool bIsFirstLineInMessage = true;
+		for (const FTextRange& LineRange : LineRanges)
 		{
-			FString Line = MessageLines[i];
-			if (Line.EndsWith(TEXT("\r")))
-			{
-				Line = Line.LeftChop(1);
-			}
+			if (LineRange.IsEmpty())
+				continue;
+
+			FString Line = CurrentLogDump.Mid(LineRange.BeginIndex, LineRange.Len());
 			Line = Line.ConvertTabsToSpaces(4);
-			OutMessages.Add(MakeShareable(new FLogMessage((i == 0) ? FOutputDevice::FormatLogLine(Verbosity, Category, *Line) : Line, Style)));
-		}
 
-		return (MessageLines.Num() > 0);
-	}
-}
+			OutMessages.Add(MakeShareable(new FLogMessage(MakeShareable(new FString((bIsFirstLineInMessage) ? FOutputDevice::FormatLogLine(Verbosity, Category, *Line) : Line)), Style)));
 
-FReply SOutputLog::OnKeyDown( const FGeometry& MyGeometry, const FKeyboardEvent& InKeyboardEvent )
-{
-	FReply Reply = FReply::Unhandled();
+			bIsFirstLineInMessage = false;
+		}
 
-	if( OutputLogActions->ProcessCommandBindings( InKeyboardEvent ) )
-	{
-		// handle the event if a command was processed
-		Reply = FReply::Handled();
+		return OldNumMessages != OutMessages.Num();
 	}
-	return Reply;
 }
 
 void SOutputLog::Serialize( const TCHAR* V, ELogVerbosity::Type Verbosity, const class FName& Category )
 {
-	if (CreateLogMessages(V, Verbosity, Category, Messages))
+	if ( MessagesTextMarshaller->AppendMessage(V, Verbosity, Category) )
 	{
-		MessageListView->RequestListRefresh();
 		// Don't scroll to the bottom automatically when the user is scrolling the view or has scrolled it away from the bottom.
-		if( OutputLogScrollBar->DistanceFromBottom() <= 0.f )
+		if( MessagesTextBox->GetVScrollBar()->DistanceFromBottom() <= 0.f )
 		{
-			MessageListView->RequestScrollIntoView(Messages.Last());
-		}
-	}
-}
+			// Force a refresh so that the message has been added before we try and jump to it
+			//MessagesTextBox->Refresh();
 
-TSharedRef<ITableRow> SOutputLog::MakeLogListItemWidget(TSharedPtr<FLogMessage> Message, const TSharedRef<STableViewBase>& OwnerTable)
-{
-	check(Message.IsValid());
-
-	return
-	SNew(STableRow< TSharedPtr<FLogMessage> >, OwnerTable)
-	.Style( FEditorStyle::Get(), Message->Style )
-	[
-		SNew(SHorizontalBox) //.ToolTipText(Message->Message)		showing the line as tooltip is a workaround to show very long lines as we don't have horizontal scrolling yet
-		+SHorizontalBox::Slot().AutoWidth() .Padding(0)
-		[
-			SNew(STextBlock) .Text(Message->Message) .TextStyle( FEditorStyle::Get(), TEXT("Log.Normal") )
-		]
-	];
-}
-
-TSharedPtr<SWidget> SOutputLog::BuildMenuWidget()
-{
-	FOutputLogModule& OutputLogModule = FModuleManager::GetModuleChecked<FOutputLogModule>( TEXT("OutputLog") );
-
-	FMenuBuilder MenuBuilder( true, OutputLogActions );
-	{ 
-		MenuBuilder.BeginSection("OutputLogEdit");
-		{
-			MenuBuilder.AddMenuEntry(FOutputLogCommands::Get().CopyOutputLog);
-			MenuBuilder.AddMenuEntry(FOutputLogCommands::Get().SelectAllInOutputLog);
-			MenuBuilder.AddMenuEntry(FOutputLogCommands::Get().SelectNoneInOutputLog);
+			MessagesTextBox->ScrollTo(FTextLocation(MessagesTextMarshaller->GetNumMessages() - 1));
 		}
-		MenuBuilder.EndSection();
-		
-		MenuBuilder.AddMenuEntry(FOutputLogCommands::Get().ClearOutputLog);
-
 	}
-
-	return MenuBuilder.MakeWidget();
 }
 
-void SOutputLog::OnCopy()
+void SOutputLog::ExtendTextBoxMenu(FMenuBuilder& Builder)
 {
-	TArray<TSharedPtr<FLogMessage>> SelectedItems = MessageListView->GetSelectedItems();
-
-	if (SelectedItems.Num())
-	{
-		// Make sure the selected range is sorted in descending index order
-		SelectedItems.Sort([this](const TSharedPtr<FLogMessage>& Item1, const TSharedPtr<FLogMessage>& Item2) -> bool
-		{
-			const int32 Item1Index = Messages.IndexOfByKey(Item1);
-			const int32 Item2Index = Messages.IndexOfByKey(Item2);
-			return Item1Index < Item2Index;
-		});
-
-		FString SelectedText;
-
-		for (int32 SelectedItemIdx = 0; SelectedItemIdx < SelectedItems.Num(); SelectedItemIdx++)
-		{
-			SelectedText += SelectedItems[SelectedItemIdx]->Message + LINE_TERMINATOR;
-		}
-
-		SelectedText = SelectedText.LeftChop(FCString::Strlen(LINE_TERMINATOR));
-
-		// Copy text to clipboard
-		FPlatformMisc::ClipboardCopy( *SelectedText );
-	}
-}
-
-bool SOutputLog::CanCopy() const
-{
-	const int32 NumSelected = MessageListView->GetNumItemsSelected();
-	return NumSelected != 0;
+	FUIAction ClearOutputLogAction(
+		FExecuteAction::CreateRaw( this, &SOutputLog::OnClearLog ),
+		FCanExecuteAction::CreateSP( this, &SOutputLog::CanClearLog )
+		);
+
+	Builder.AddMenuEntry(
+		NSLOCTEXT("OutputLog", "ClearLogLabel", "Clear Log"), 
+		NSLOCTEXT("OutputLog", "ClearLogTooltip", "Clears all log messages"), 
+		FSlateIcon(), 
+		ClearOutputLogAction
+		);
 }
 
 void SOutputLog::OnClearLog()
 {
-	Messages.Empty();
-	MessageListView->RequestListRefresh();
-}
+	// Make sure the cursor is back at the start of the log before we clear it
+	MessagesTextBox->GoTo(FTextLocation(0));
 
-bool SOutputLog::CanClearLog() const
-{
-	return Messages.Num() != 0;
+	MessagesTextMarshaller->ClearMessages();
+	MessagesTextBox->Refresh();
 }
 
-void SOutputLog::OnSelectAll()
-{
-	for (int32 ItemIdx = 0; ItemIdx < Messages.Num(); ItemIdx++)
-	{
-		if (!MessageListView->IsItemSelected(Messages[ItemIdx]))
-		{
-			MessageListView->SetItemSelection(Messages[ItemIdx], true);
-		}
-	}
-}
-
-bool SOutputLog::CanSelectAll() const
-{
-	const int32 NumSelected = MessageListView->GetNumItemsSelected();
-	return Messages.Num() && NumSelected != Messages.Num();
-}
-
-void SOutputLog::OnSelectNone()
-{
-	for (int32 ItemIdx = 0; ItemIdx < Messages.Num(); ItemIdx++)
-	{
-		if (MessageListView->IsItemSelected(Messages[ItemIdx]))
-		{
-			MessageListView->SetItemSelection(Messages[ItemIdx], false);
-		}
-	}
-}
-
-bool SOutputLog::CanSelectNone() const
+bool SOutputLog::CanClearLog() const
 {
-	const int32 NumSelected = MessageListView->GetNumItemsSelected();
-	return Messages.Num() && NumSelected;
+	return MessagesTextMarshaller->GetNumMessages() > 0;
 }
diff --git a/Engine/Source/Developer/OutputLog/Private/SOutputLog.h b/Engine/Source/Developer/OutputLog/Private/SOutputLog.h
--- a/Engine/Source/Developer/OutputLog/Private/SOutputLog.h
+++ b/Engine/Source/Developer/OutputLog/Private/SOutputLog.h
@@ -3,16 +3,19 @@
 #pragma once
 
 
+class FOutputLogTextLayoutMarshaller;
+
+
 /**
  * A single log message for the output log, holding a message and
  * a style, for color and bolding of the message.
  */
 struct FLogMessage
 {
-	FString Message;
+	TSharedRef<FString> Message;
 	FName Style;
 
-	FLogMessage(const FString& NewMessage, FName NewStyle = NAME_None)
+	FLogMessage(const TSharedRef<FString>& NewMessage, FName NewStyle = NAME_None)
 		: Message(NewMessage)
 		, Style(NewStyle)
 	{
@@ -153,60 +156,15 @@ public:
 	 */
 	static bool CreateLogMessages( const TCHAR* V, ELogVerbosity::Type Verbosity, const class FName& Category, TArray< TSharedPtr<FLogMessage> >& OutMessages );
 
-	/**
-	 * Called after a key is pressed when this widget has keyboard focus
-	 *
-	 * @param MyGeometry The Geometry of the widget receiving the event
-	 * @param  InKeyboardEvent  Keyboard event
-	 *
-	 * @return  Returns whether the event was handled, along with other possible actions
-	 */
-	virtual FReply OnKeyDown( const FGeometry& MyGeometry, const FKeyboardEvent& InKeyboardEvent ) override;
-
 protected:
 
 	virtual void Serialize( const TCHAR* V, ELogVerbosity::Type Verbosity, const class FName& Category ) override;
 
-	/** Makes the widget for the log messages in the list view */
-	TSharedRef<ITableRow> MakeLogListItemWidget(TSharedPtr<FLogMessage> Message, const TSharedRef<STableViewBase>& OwnerTable);
-
 private:
 	/**
-	 * Creates a widget for the context menu that can be inserted into a pop-up window
-	 *
-	 * @return	Widget for this context menu
-	 */
-	TSharedPtr< SWidget > BuildMenuWidget();
-
-	/**
-	 * Called when copy is selected
-	 */
-	void OnCopy();
-
-	/**
-	 * Called to determine whether copy is currently a valid command
-	 */
-	bool CanCopy() const;
-
-	/**
-	 * Called when select all is selected
-	 */
-	void OnSelectAll();
-
-	/**
-	 * Called to determine whether select all is currently a valid command
-	 */
-	bool CanSelectAll() const;
-
-	/**
-	 * Called when select none is selected
+	 * Extends the context menu used by the text box
 	 */
-	void OnSelectNone();
-
-	/**
-	 * Called to determine whether select none is currently a valid command
-	 */
-	bool CanSelectNone() const;
+	void ExtendTextBoxMenu(FMenuBuilder& Builder);
 
 	/**
 	 * Called when delete all is selected
@@ -218,17 +176,9 @@ private:
 	 */
 	bool CanClearLog() const;
 
-	/** 
-	 * Output log commands
-	 */
-	TSharedPtr<FUICommandList> OutputLogActions;
-
-	/** All log messages stored in this widget for the list view */
-	TArray< TSharedPtr<FLogMessage> > Messages;
-
-	/** The list view for showing all log messages. Should be replaced by a full text editor */
-	TSharedPtr< SListView< TSharedPtr<FLogMessage> > > MessageListView;
+	/** Converts the array of messages into something the text box understands */
+	TSharedPtr< FOutputLogTextLayoutMarshaller > MessagesTextMarshaller;
 
-	/** Scroll bar for output log. */
-	TSharedPtr< SScrollBar > OutputLogScrollBar;
+	/** The editable text showing all log messages */
+	TSharedPtr< SMultiLineEditableTextBox > MessagesTextBox;
 };
diff --git a/Engine/Source/Editor/EditorStyle/Private/SlateEditorStyle.cpp b/Engine/Source/Editor/EditorStyle/Private/SlateEditorStyle.cpp
--- a/Engine/Source/Editor/EditorStyle/Private/SlateEditorStyle.cpp
+++ b/Engine/Source/Editor/EditorStyle/Private/SlateEditorStyle.cpp
@@ -863,27 +863,35 @@ void FSlateEditorStyle::FStyle::SetupGeneralStyles()
 	// Output Log Window
 #if WITH_EDITOR || IS_PROGRAM
 	{
-		Set("LogTableRow.Normal", FTableRowStyle(NormalTableRowStyle)
-			.SetTextColor( FLinearColor(FColor(0xffaaaaaa)) )
-			.SetSelectedTextColor( FLinearColor(FColor(0xff333333)) ) );
+		const FTextBlockStyle NormalLogText = FTextBlockStyle(NormalText)
+			.SetFont( TTF_FONT( "Fonts/DroidSansMono", 9 ) )
+			.SetColorAndOpacity( FLinearColor(FColor(0xffaaaaaa)) )
+			//.SetSelectedBackgroundColor( SelectionColor );
+			.SetSelectedBackgroundColor( FLinearColor(FColor(0xffeeeeee)) );
 
-		Set("LogTableRow.Command", FTableRowStyle(NormalTableRowStyle)
-			.SetTextColor( FLinearColor(FColor(0xff33ff33)) )
-			.SetSelectedTextColor( FLinearColor(FColor(0xff055505)) ) );
+		Set("Log.Normal", NormalLogText );
 
-		Set("LogTableRow.Warning", FTableRowStyle(NormalTableRowStyle)
-			.SetTextColor( FLinearColor(FColor(0xffbbbb44)) )
-			.SetSelectedTextColor( FLinearColor(FColor(0xff444411)) ) );
+		Set("Log.Command", FTextBlockStyle(NormalLogText)
+			.SetColorAndOpacity( FLinearColor(FColor(0xff33dd33)) )
+			);
 
-		Set("LogTableRow.Error", FTableRowStyle(NormalTableRowStyle)
-			.SetTextColor( FLinearColor(FColor(0xffff0000)) )
-			.SetSelectedTextColor( FLinearColor(FColor(0xff440000)) ) );
+		Set("Log.Warning", FTextBlockStyle(NormalLogText)
+			.SetColorAndOpacity( FLinearColor(FColor(0xffbbbb44)) )
+			);
 
-		Set("Log.Normal", FTextBlockStyle(NormalText)
-			.SetFont( TTF_FONT( "Fonts/DroidSansMono", 9 ) ) );
+		Set("Log.Error", FTextBlockStyle(NormalLogText)
+			.SetColorAndOpacity( FLinearColor(FColor(0xffdd0000)) )
+			);
 
 		Set("Log.TabIcon", new IMAGE_BRUSH( "Icons/icon_tab_OutputLog_16x", Icon16x16 ) );
 
+		Set("Log.TextBox", FEditableTextBoxStyle(NormalEditableTextBoxStyle)
+			.SetBackgroundImageNormal( BOX_BRUSH( "Common/GroupBorder", FMargin(4.0f/16.0f) ) )
+			.SetBackgroundImageHovered( BOX_BRUSH( "Common/GroupBorder", FMargin(4.0f/16.0f) ) )
+			.SetBackgroundImageFocused( BOX_BRUSH( "Common/GroupBorder", FMargin(4.0f/16.0f) ) )
+			.SetBackgroundImageReadOnly( BOX_BRUSH( "Common/GroupBorder", FMargin(4.0f/16.0f) ) )
+			);
+
 		Set("DebugConsole.Background", new BOX_BRUSH("Old/Menu_Background", FMargin(8.0f / 64.0f)));
 	}
 
diff --git a/Engine/Source/Runtime/Slate/Private/Widgets/Input/SMultiLineEditableTextBox.cpp b/Engine/Source/Runtime/Slate/Private/Widgets/Input/SMultiLineEditableTextBox.cpp
--- a/Engine/Source/Runtime/Slate/Private/Widgets/Input/SMultiLineEditableTextBox.cpp
+++ b/Engine/Source/Runtime/Slate/Private/Widgets/Input/SMultiLineEditableTextBox.cpp
@@ -81,12 +81,14 @@ void SMultiLineEditableTextBox::Construct( const FArguments& InArgs )
 				[
 					SAssignNew( EditableText, SMultiLineEditableText )
 					.Text( InArgs._Text )
+					.TextStyle( InArgs._TextStyle )
 					.Marshaller( InArgs._Marshaller )
 					.Font( Font )
 					.IsReadOnly( InArgs._IsReadOnly )
 					.OnTextChanged( InArgs._OnTextChanged )
 					.OnTextCommitted( InArgs._OnTextCommitted )
 					.OnCursorMoved( InArgs._OnCursorMoved )
+					.ContextMenuExtender( InArgs._ContextMenuExtender )
 					.Justification(InArgs._Justification)
 					.LineHeightPercentage(InArgs._LineHeightPercentage)
 					.Margin(InArgs._Margin)
@@ -230,6 +232,16 @@ void SMultiLineEditableTextBox::InsertRunAtCursor(TSharedRef<IRun> InRun)
 	EditableText->InsertRunAtCursor(InRun);
 }
 
+void SMultiLineEditableTextBox::GoTo(const FTextLocation& NewLocation)
+{
+	return EditableText->GoTo(NewLocation);
+}
+
+void SMultiLineEditableTextBox::ScrollTo(const FTextLocation& NewLocation)
+{
+	return EditableText->ScrollTo(NewLocation);
+}
+
 void SMultiLineEditableTextBox::ApplyToSelection(const FRunInfo& InRunInfo, const FTextBlockStyle& InStyle)
 {
 	EditableText->ApplyToSelection(InRunInfo, InStyle);
@@ -245,4 +257,19 @@ const TArray<TSharedRef<const IRun>> SMultiLineEditableTextBox::GetSelectedRuns(
 	return EditableText->GetSelectedRuns();
 }
 
+TSharedPtr<const SScrollBar> SMultiLineEditableTextBox::GetHScrollBar() const
+{
+	return EditableText->GetHScrollBar();
+}
+
+TSharedPtr<const SScrollBar> SMultiLineEditableTextBox::GetVScrollBar() const
+{
+	return EditableText->GetVScrollBar();
+}
+
+void SMultiLineEditableTextBox::Refresh()
+{
+	return EditableText->Refresh();
+}
+
 #endif //WITH_FANCY_TEXT
diff --git a/Engine/Source/Runtime/Slate/Private/Widgets/Text/SMultiLineEditableText.cpp b/Engine/Source/Runtime/Slate/Private/Widgets/Text/SMultiLineEditableText.cpp
--- a/Engine/Source/Runtime/Slate/Private/Widgets/Text/SMultiLineEditableText.cpp
+++ b/Engine/Source/Runtime/Slate/Private/Widgets/Text/SMultiLineEditableText.cpp
@@ -153,9 +153,10 @@ int32 SMultiLineEditableText::FTextSelectionRunRenderer::OnPaint( const FPaintAr
 	FVector2D Location( Block->GetLocationOffset() );
 	Location.Y = Line.Offset.Y;
 
-	const FLinearColor FocusedSelectionColorAdjustment(-0.2f, -0.05f, 0.15f);
-	const FLinearColor UnfocusedSelectionColorAdjustment(-0.01f, -0.01f, 0.01f);
-	const FColor SelectionBackgroundColorAndOpacity = ((FLinearColor::White - DefaultStyle.ColorAndOpacity.GetColor(InWidgetStyle))*0.5f + (bHasKeyboardFocus ? FocusedSelectionColorAdjustment : UnfocusedSelectionColorAdjustment)) * InWidgetStyle.GetColorAndOpacityTint();
+	// If we've not been set to an explicit color, calculate a suitable one from the linked color
+	const FColor SelectionBackgroundColorAndOpacity = DefaultStyle.SelectedBackgroundColor.IsColorSpecified() 
+		? DefaultStyle.SelectedBackgroundColor.GetSpecifiedColor() * InWidgetStyle.GetColorAndOpacityTint()
+		: ((FLinearColor::White - DefaultStyle.SelectedBackgroundColor.GetColor(InWidgetStyle))*0.5f + FLinearColor(-0.2f, -0.05f, 0.15f)) * InWidgetStyle.GetColorAndOpacityTint();
 
 	// The block size and offset values are pre-scaled, so we need to account for that when converting the block offsets into paint geometry
 	const float InverseScale = Inverse(AllottedGeometry.Scale);
@@ -170,16 +171,18 @@ int32 SMultiLineEditableText::FTextSelectionRunRenderer::OnPaint( const FPaintAr
 			AllottedGeometry.ToPaintGeometry(TransformVector(InverseScale, FVector2D(HighlightWidth, Line.Size.Y)), FSlateLayoutTransform(TransformPoint(InverseScale, Location))),
 			&DefaultStyle.HighlightShape,
 			MyClippingRect,
-			bParentEnabled ? ESlateDrawEffect::None : ESlateDrawEffect::DisabledEffect,
+			bParentEnabled && bHasKeyboardFocus ? ESlateDrawEffect::None : ESlateDrawEffect::DisabledEffect,
 			SelectionBackgroundColorAndOpacity
 		);
 	}
 
-	//FLinearColor InvertedForeground = FLinearColor::White - InWidgetStyle.GetForegroundColor();
-	//InvertedForeground.A = InWidgetStyle.GetForegroundColor().A;
+	/*
+	FLinearColor InvertedForeground = FLinearColor::White - InWidgetStyle.GetForegroundColor();
+	InvertedForeground.A = InWidgetStyle.GetForegroundColor().A;
 
-	//FWidgetStyle WidgetStyle( InWidgetStyle );
-	//WidgetStyle.SetForegroundColor( InvertedForeground );
+	FWidgetStyle WidgetStyle( InWidgetStyle );
+	WidgetStyle.SetForegroundColor( InvertedForeground );
+	*/
 
 	return Run->OnPaint( Args, Line, Block, DefaultStyle, AllottedGeometry, MyClippingRect, OutDrawElements, LayerId, InWidgetStyle, bParentEnabled );
 }
@@ -195,7 +198,6 @@ SMultiLineEditableText::SMultiLineEditableText()
 	, bIsDragSelecting(false)
 	, bWasFocusedByLastMouseDown(false)
 	, bHasDragSelectedSinceFocused(false)
-	, bPendingScrollCursorIntoView(false)
 	, CurrentUndoLevel(INDEX_NONE)
 	, bIsChangingText(false)
 	, IsReadOnly(false)
@@ -368,19 +370,18 @@ bool SMultiLineEditableText::SetEditableText(const FText& TextToSet, const bool
 		ClearSelection();
 		TextLayout->ClearLines();
 
-		if(TextToSetString.Len())
-		{
-			Marshaller->SetText(TextToSetString, *TextLayout);
-		}
-		else
+		Marshaller->SetText(TextToSetString, *TextLayout);
+
+		const TArray< FTextLayout::FLineModel >& Lines = TextLayout->GetLineModels();
+		if(Lines.Num() == 0)
 		{
-			TSharedPtr< FString > LineText = MakeShareable(new FString());
+			TSharedRef< FString > LineText = MakeShareable(new FString());
 
 			// Create an empty run
 			TArray< TSharedRef< IRun > > Runs;
-			Runs.Add(FSlateTextRun::Create(FRunInfo(), LineText.ToSharedRef(), TextStyle));
+			Runs.Add(FSlateTextRun::Create(FRunInfo(), LineText, TextStyle));
 
-			TextLayout->AddLine(LineText.ToSharedRef(), Runs);
+			TextLayout->AddLine(LineText, Runs);
 		}
 
 		return true;
@@ -408,6 +409,8 @@ void SMultiLineEditableText::ForceRefreshTextLayout(const FText& CurrentText)
 	SelectionStart = OldSelectionStart;
 	CursorInfo = OldCursorInfo;
 	UpdateCursorHighlight();
+
+	TextLayout->UpdateIfNeeded();
 }
 
 void SMultiLineEditableText::OnHScrollBarMoved(const float InScrollOffsetFraction)
@@ -475,6 +478,10 @@ void SMultiLineEditableText::OnKeyboardFocusLost( const FKeyboardFocusEvent& InK
 
 		OnTextCommitted.ExecuteIfBound(EditedText, TextAction);
 		UpdateCursorHighlight();
+
+		// UpdateCursorHighlight always tries to scroll to the cursor, but we don't want that to happen when we 
+		// lose focus since it can cause the scroll position to jump unexpectedly
+		PositionToScrollIntoView = TOptional<FScrollInfo>();
 	}
 }
 
@@ -931,7 +938,7 @@ FReply SMultiLineEditableText::MoveCursor( FMoveCursor Args )
 
 void SMultiLineEditableText::UpdateCursorHighlight()
 {
-	bPendingScrollCursorIntoView = true;
+	PositionToScrollIntoView = FScrollInfo(CursorInfo.GetCursorInteractionLocation(), CursorInfo.GetCursorAlignment());
 
 	RemoveCursorHighlight();
 
@@ -1241,6 +1248,37 @@ void SMultiLineEditableText::InsertRunAtCursor(TSharedRef<IRun> InRun)
 	FinishChangingText();
 }
 
+void SMultiLineEditableText::GoTo(const FTextLocation& NewLocation)
+{
+	const TArray< FTextLayout::FLineModel >& Lines = TextLayout->GetLineModels();
+	if (Lines.IsValidIndex(NewLocation.GetLineIndex()))
+	{
+		const FTextLayout::FLineModel& Line = Lines[NewLocation.GetLineIndex()];
+		if (NewLocation.GetOffset() <= Line.Text->Len())
+		{
+			ClearSelection();
+
+			CursorInfo.SetCursorLocationAndCalculateAlignment(TextLayout, NewLocation);
+			OnCursorMoved.ExecuteIfBound(CursorInfo.GetCursorInteractionLocation());
+			UpdatePreferredCursorScreenOffsetInLine();
+			UpdateCursorHighlight();
+		}
+	}
+}
+
+void SMultiLineEditableText::ScrollTo(const FTextLocation& NewLocation)
+{
+	const TArray< FTextLayout::FLineModel >& Lines = TextLayout->GetLineModels();
+	if (Lines.IsValidIndex(NewLocation.GetLineIndex()))
+	{
+		const FTextLayout::FLineModel& Line = Lines[NewLocation.GetLineIndex()];
+		if (NewLocation.GetOffset() <= Line.Text->Len())
+		{
+			PositionToScrollIntoView = FScrollInfo(NewLocation, ECursorAlignment::Left);
+		}
+	}
+}
+
 void SMultiLineEditableText::ApplyToSelection(const FRunInfo& InRunInfo, const FTextBlockStyle& InStyle)
 {
 	StartChangingText();
@@ -1419,6 +1457,41 @@ const TArray<TSharedRef<const IRun>> SMultiLineEditableText::GetSelectedRuns() c
 	return Runs;
 }
 
+TSharedPtr<const SScrollBar> SMultiLineEditableText::GetHScrollBar() const
+{
+	return HScrollBar;
+}
+
+TSharedPtr<const SScrollBar> SMultiLineEditableText::GetVScrollBar() const
+{
+	return VScrollBar;
+}
+
+void SMultiLineEditableText::Refresh()
+{
+	bool bHasSetText = false;
+
+	const FText& TextToSet = BoundText.Get(FText::GetEmpty());
+	if (!BoundTextLastTick.IdenticalTo(TextToSet))
+	{
+		// The pointer used by the bound text has changed, however the text may still be the same - check that now
+		if (!BoundTextLastTick.IsDisplayStringEqualTo(TextToSet))
+		{
+			// The source text has changed, so update the internal editable text
+			bHasSetText = true;
+			SetEditableText(TextToSet);
+		}
+
+		// Update this even if the text is lexically identical, as it will update the pointer compared by IdenticalTo for the next Tick
+		BoundTextLastTick = FTextSnapshot(TextToSet);
+	}
+
+	if (!bHasSetText && Marshaller->IsDirty())
+	{
+		ForceRefreshTextLayout(TextToSet);
+	}
+}
+
 bool SMultiLineEditableText::CanExecuteSelectAll() const
 {
 	bool bCanExecute = true;
@@ -1730,7 +1803,7 @@ void SMultiLineEditableText::InsertTextAtCursorImpl(const FString& InString)
 
 bool SMultiLineEditableText::CanExecuteUndo() const
 {
-	if(TextInputMethodContext->IsComposing)
+	if(IsReadOnly.Get() || TextInputMethodContext->IsComposing)
 	{
 		return false;
 	}
@@ -1957,45 +2030,25 @@ void SMultiLineEditableText::Tick( const FGeometry& AllottedGeometry, const doub
 
 	if (!bShouldAppearFocused)
 	{
-		bool bHasSetText = false;
-
-		const FText& TextToSet = BoundText.Get(FText::GetEmpty());
-		if (!BoundTextLastTick.IdenticalTo(TextToSet))
-		{
-			// The pointer used by the bound text has changed, however the text may still be the same - check that now
-			if (!BoundTextLastTick.IsDisplayStringEqualTo(TextToSet))
-			{
-				// The source text has changed, so update the internal editable text
-				bHasSetText = true;
-				SetEditableText(TextToSet);
-			}
-
-			// Update this even if the text is lexically identical, as it will update the pointer compared by IdenticalTo for the next Tick
-			BoundTextLastTick = FTextSnapshot(TextToSet);
-		}
-
-		if (!bHasSetText && Marshaller->IsDirty())
-		{
-			ForceRefreshTextLayout(TextToSet);
-		}
+		Refresh();
 	}
 
 	const float FontMaxCharHeight = FTextEditHelper::GetFontHeight(TextStyle.Font);
 	const float CaretWidth = FTextEditHelper::CalculateCaretWidth(FontMaxCharHeight);
 
 	// Try and make sure that the line containing the cursor is in view
-	if (bPendingScrollCursorIntoView)
+	if (PositionToScrollIntoView.IsSet())
 	{
-		bPendingScrollCursorIntoView = false;
+		const FScrollInfo& ScrollInfo = PositionToScrollIntoView.GetValue();
 
 		const TArray<FTextLayout::FLineView>& LineViews = TextLayout->GetLineViews();
-		const int32 LineViewIndex = TextLayout->GetLineViewIndexForTextLocation(LineViews, CursorInfo.GetCursorInteractionLocation(), CursorInfo.GetCursorAlignment() == ECursorAlignment::Right);
+		const int32 LineViewIndex = TextLayout->GetLineViewIndexForTextLocation(LineViews, ScrollInfo.Position, ScrollInfo.Alignment == ECursorAlignment::Right);
 		if (LineViews.IsValidIndex(LineViewIndex))
 		{
 			const FTextLayout::FLineView& LineView = LineViews[LineViewIndex];
 			const FSlateRect LocalLineViewRect(LineView.Offset / TextLayout->GetScale(), (LineView.Offset + LineView.Size) / TextLayout->GetScale());
 
-			const FVector2D LocalCursorLocation = TextLayout->GetLocationAt(CursorInfo.GetCursorInteractionLocation(), CursorInfo.GetCursorAlignment() == ECursorAlignment::Right) / TextLayout->GetScale();
+			const FVector2D LocalCursorLocation = TextLayout->GetLocationAt(ScrollInfo.Position, ScrollInfo.Alignment == ECursorAlignment::Right) / TextLayout->GetScale();
 			const FSlateRect LocalCursorRect(LocalCursorLocation, FVector2D(LocalCursorLocation.X + CaretWidth, LocalCursorLocation.Y + FontMaxCharHeight));
 
 			if (LocalCursorRect.Left < 0.0f)
@@ -2016,6 +2069,8 @@ void SMultiLineEditableText::Tick( const FGeometry& AllottedGeometry, const doub
 				ScrollOffset.Y += (LocalLineViewRect.Bottom - AllottedGeometry.Size.Y);
 			}
 		}
+
+		PositionToScrollIntoView = TOptional<FScrollInfo>();
 	}
 
 	{
diff --git a/Engine/Source/Runtime/Slate/Private/Widgets/Text/TextBlockLayout.cpp b/Engine/Source/Runtime/Slate/Private/Widgets/Text/TextBlockLayout.cpp
--- a/Engine/Source/Runtime/Slate/Private/Widgets/Text/TextBlockLayout.cpp
+++ b/Engine/Source/Runtime/Slate/Private/Widgets/Text/TextBlockLayout.cpp
@@ -158,6 +158,7 @@ bool FTextBlockLayout::IsStyleUpToDate(const FTextBlockStyle& NewStyle) const
 		&& (CurrentStyle.ColorAndOpacity == NewStyle.ColorAndOpacity)
 		&& (CurrentStyle.ShadowOffset == NewStyle.ShadowOffset)
 		&& (CurrentStyle.ShadowColorAndOpacity == NewStyle.ShadowColorAndOpacity)
+		&& (CurrentStyle.SelectedBackgroundColor == NewStyle.SelectedBackgroundColor)
 		&& (CurrentStyle.HighlightColor == NewStyle.HighlightColor)
 		&& (CurrentStyle.HighlightShape == NewStyle.HighlightShape);
 }
diff --git a/Engine/Source/Runtime/Slate/Public/Widgets/Input/SMultiLineEditableTextBox.h b/Engine/Source/Runtime/Slate/Public/Widgets/Input/SMultiLineEditableTextBox.h
--- a/Engine/Source/Runtime/Slate/Public/Widgets/Input/SMultiLineEditableTextBox.h
+++ b/Engine/Source/Runtime/Slate/Public/Widgets/Input/SMultiLineEditableTextBox.h
@@ -13,7 +13,8 @@ class SLATE_API SMultiLineEditableTextBox : public SBorder
 public:
 
 	SLATE_BEGIN_ARGS( SMultiLineEditableTextBox )
-		: _Style(&FCoreStyle::Get().GetWidgetStyle< FEditableTextBoxStyle >("NormalEditableTextBox"))
+		: _Style(&FCoreStyle::Get().GetWidgetStyle<FEditableTextBoxStyle>("NormalEditableTextBox"))
+		, _TextStyle(&FCoreStyle::Get().GetWidgetStyle<FTextBlockStyle>("NormalText"))
 		, _Marshaller()
 		, _Text()
 		, _HintText()
@@ -40,6 +41,9 @@ public:
 		/** The styling of the textbox */
 		SLATE_STYLE_ARGUMENT( FEditableTextBoxStyle, Style )
 
+		/** Pointer to a style of the text block, which dictates the font, color, and shadow options. */
+		SLATE_STYLE_ARGUMENT( FTextBlockStyle, TextStyle )
+
 		/** The marshaller used to get/set the raw text to/from the text layout. */
 		SLATE_ARGUMENT(TSharedPtr< ITextLayoutMarshaller >, Marshaller)
 
@@ -103,6 +107,9 @@ public:
 		/** Called when the cursor is moved within the text area */
 		SLATE_EVENT( SMultiLineEditableText::FOnCursorMoved, OnCursorMoved )
 
+		/** Menu extender for the right-click context menu */
+		SLATE_EVENT( FMenuExtensionDelegate, ContextMenuExtender )
+
 		/** Whether text wraps onto a new line when it's length exceeds this width; if this value is zero or negative, no wrapping occurs. */
 		SLATE_ATTRIBUTE( float, WrapTextAt )
 
@@ -175,6 +182,12 @@ public:
 	/** Insert the given run at the current cursor position */
 	void InsertRunAtCursor(TSharedRef<IRun> InRun);
 
+	/** Move the cursor to the given location in the document */
+	void GoTo(const FTextLocation& NewLocation);
+
+	/** Scroll to the given location in the document (without moving the cursor) */
+	void ScrollTo(const FTextLocation& NewLocation);
+
 	/** Apply the given style to the currently selected text (or insert a new run at the current cursor position if no text is selected) */
 	void ApplyToSelection(const FRunInfo& InRunInfo, const FTextBlockStyle& InStyle);
 
@@ -184,6 +197,15 @@ public:
 	/** Get the runs currently that are current selected, some of which may be only partially selected */
 	const TArray<TSharedRef<const IRun>> GetSelectedRuns() const;
 
+	/** Get the horizontal scroll bar widget */
+	TSharedPtr<const SScrollBar> GetHScrollBar() const;
+
+	/** Get the vertical scroll bar widget */
+	TSharedPtr<const SScrollBar> GetVScrollBar() const;
+
+	/** Refresh this text box immediately, rather than wait for the usual caching mechanisms to take affect on the text Tick */
+	void Refresh();
+
 protected:
 
 	/** Editable text widget */
diff --git a/Engine/Source/Runtime/Slate/Public/Widgets/Text/SMultiLineEditableText.h b/Engine/Source/Runtime/Slate/Public/Widgets/Text/SMultiLineEditableText.h
--- a/Engine/Source/Runtime/Slate/Public/Widgets/Text/SMultiLineEditableText.h
+++ b/Engine/Source/Runtime/Slate/Public/Widgets/Text/SMultiLineEditableText.h
@@ -115,6 +115,12 @@ public:
 	/** Insert the given run at the current cursor position */
 	void InsertRunAtCursor(TSharedRef<IRun> InRun);
 
+	/** Move the cursor to the given location in the document (will also scroll to this point) */
+	void GoTo(const FTextLocation& NewLocation);
+
+	/** Scroll to the given location in the document (without moving the cursor) */
+	void ScrollTo(const FTextLocation& NewLocation);
+
 	/** Apply the given style to the currently selected text (or insert a new run at the current cursor position if no text is selected) */
 	void ApplyToSelection(const FRunInfo& InRunInfo, const FTextBlockStyle& InStyle);
 
@@ -124,6 +130,15 @@ public:
 	/** Get the runs currently that are current selected, some of which may be only partially selected */
 	const TArray<TSharedRef<const IRun>> GetSelectedRuns() const;
 
+	/** Get the horizontal scroll bar widget */
+	TSharedPtr<const SScrollBar> GetHScrollBar() const;
+
+	/** Get the vertical scroll bar widget */
+	TSharedPtr<const SScrollBar> GetVScrollBar() const;
+
+	/** Refresh this text box immediately, rather than wait for the usual caching mechanisms to take affect on the text Tick */
+	void Refresh();
+
 private:
 	
 	/** 
@@ -262,6 +277,28 @@ private:
 		}
 	};
 
+	/** Information needed to be able to scroll to a given point */
+	struct FScrollInfo
+	{
+		FScrollInfo()
+			: Position()
+			, Alignment(ECursorAlignment::Left)
+		{
+		}
+
+		FScrollInfo(const FTextLocation InPosition, const ECursorAlignment InAlignment)
+			: Position(InPosition)
+			, Alignment(InAlignment)
+		{
+		}
+
+		/** The location in the document to scroll to (in line model space) */
+		FTextLocation Position;
+
+		/** The alignment at the given position. This may affect which line view the character maps to when converted from line model space */
+		ECursorAlignment Alignment;
+	};
+
 	/** Run highlighter used to draw the cursor */
 	class FCursorLineHighlighter : public ISlateLineHighlighter
 	{
@@ -539,8 +576,8 @@ private:
 	/** True if characters were selected by dragging since the last keyboard focus.  Used for text selection. */
 	bool bHasDragSelectedSinceFocused;
 
-	/** True if we're pending a check to ensure the cursor is currently scrolled into view */
-	bool bPendingScrollCursorIntoView;
+	/** If set, the pending data containing a position that should be scrolled into view */
+	TOptional< FScrollInfo > PositionToScrollIntoView;
 
 	/** Undo states */
 	TArray< FUndoState > UndoStates;
diff --git a/Engine/Source/Runtime/SlateCore/Private/Styling/SlateTypes.cpp b/Engine/Source/Runtime/SlateCore/Private/Styling/SlateTypes.cpp
--- a/Engine/Source/Runtime/SlateCore/Private/Styling/SlateTypes.cpp
+++ b/Engine/Source/Runtime/SlateCore/Private/Styling/SlateTypes.cpp
@@ -63,6 +63,7 @@ void FCheckBoxStyle::PostSerialize(const FArchive& Ar)
 FTextBlockStyle::FTextBlockStyle()
 : Font()
 , ColorAndOpacity()
+, SelectedBackgroundColor(FSlateColor::UseForeground())
 , ShadowOffset(FVector2D::ZeroVector)
 , ShadowColorAndOpacity(FLinearColor::Black)
 {
diff --git a/Engine/Source/Runtime/SlateCore/Public/Styling/SlateTypes.h b/Engine/Source/Runtime/SlateCore/Public/Styling/SlateTypes.h
--- a/Engine/Source/Runtime/SlateCore/Public/Styling/SlateTypes.h
+++ b/Engine/Source/Runtime/SlateCore/Public/Styling/SlateTypes.h
@@ -196,6 +196,11 @@ struct SLATECORE_API FTextBlockStyle : public FSlateWidgetStyle
 	FLinearColor ShadowColorAndOpacity;
 	FTextBlockStyle& SetShadowColorAndOpacity(const FLinearColor& InShadowColorAndOpacity) { ShadowColorAndOpacity = InShadowColorAndOpacity; return *this; }
 
+	/** The background color of selected text */
+	UPROPERTY(EditAnywhere, Category=Appearance)
+	FSlateColor SelectedBackgroundColor;
+	FTextBlockStyle& SetSelectedBackgroundColor(const FSlateColor& InSelectedBackgroundColor) { SelectedBackgroundColor = InSelectedBackgroundColor; return *this; }
+
 	/** The color of highlighted text */
 	UPROPERTY(EditAnywhere, BlueprintReadWrite, Category=Appearance)
 	FLinearColor HighlightColor;
