Fixed auto-scrolling not working after clearing the output log

diff --git a/Engine/Source/Developer/OutputLog/Private/SOutputLog.cpp b/Engine/Source/Developer/OutputLog/Private/SOutputLog.cpp
--- a/Engine/Source/Developer/OutputLog/Private/SOutputLog.cpp
+++ b/Engine/Source/Developer/OutputLog/Private/SOutputLog.cpp
@@ -598,10 +598,18 @@ bool FOutputLogTextLayoutMarshaller::AppendMessage(const TCHAR* InText, const EL
 	TArray< TSharedPtr<FLogMessage> > NewMessages;
 	if(SOutputLog::CreateLogMessages(InText, InVerbosity, InCategory, NewMessages))
 	{
+		const bool bWasEmpty = Messages.Num() == 0;
 		Messages.Append(NewMessages);
 
 		if(TextLayout)
 		{
+			// If we were previously empty, then we'd have inserted a dummy empty line into the document
+			// We need to remove this line now as it would cause the message indices to get out-of-sync with the line numbers, which would break auto-scrolling
+			if(bWasEmpty)
+			{
+				TextLayout->ClearLines();
+			}
+
 			// If we've already been given a text layout, then append these new messages rather than force a refresh of the entire document
 			for(const auto& Message : NewMessages)
 			{
