Some usability improvements for the output log

1) The selected text colour is now easier on the eyes
2) The output log will always scroll to the end when you enter a command
3) The output log will now correctly scroll to the bottom when it is first opened (it forces the scrollbars to always visible to avoid the horizontal scrollbar being able to move the vertical scrollbar up slightly from the bottom of the log)

diff --git a/Engine/Source/Developer/OutputLog/Private/SOutputLog.cpp b/Engine/Source/Developer/OutputLog/Private/SOutputLog.cpp
--- a/Engine/Source/Developer/OutputLog/Private/SOutputLog.cpp
+++ b/Engine/Source/Developer/OutputLog/Private/SOutputLog.cpp
@@ -659,6 +659,9 @@ FOutputLogTextLayoutMarshaller::FOutputLogTextLayoutMarshaller(TArray< TSharedPt
 BEGIN_SLATE_FUNCTION_BUILD_OPTIMIZATION
 void SOutputLog::Construct( const FArguments& InArgs )
 {
+	// Force scroll in the first Tick() once we've been sized and populated correctly
+	bPendingForceScroll = true;
+
 	MessagesTextMarshaller = FOutputLogTextLayoutMarshaller::Create(MoveTemp(InArgs._Messages));
 
 	MessagesTextBox = SNew(SMultiLineEditableTextBox)
@@ -666,6 +669,7 @@ void SOutputLog::Construct( const FArguments& InArgs )
 		.TextStyle(FEditorStyle::Get(), "Log.Normal")
 		.Marshaller(MessagesTextMarshaller)
 		.IsReadOnly(true)
+		.AlwaysShowScrollbars(true)
 		.ContextMenuExtender(this, &SOutputLog::ExtendTextBoxMenu);
 
 	ChildSlot
@@ -684,6 +688,7 @@ void SOutputLog::Construct( const FArguments& InArgs )
 			.Padding(FMargin(0.0f, 4.0f, 0.0f, 0.0f))
 			[
 				SNew( SConsoleInputBox )
+				.OnConsoleCommandExecuted(this, &SOutputLog::OnConsoleCommandExecuted)
 
 				// Always place suggestions above the input line for the output log widget
 				.SuggestionListPlacement( MenuPlacement_AboveAnchor )
@@ -691,12 +696,6 @@ void SOutputLog::Construct( const FArguments& InArgs )
 	];
 	
 	GLog->AddOutputDevice(this);
-
-	// If there's already been messages logged, scroll down to the last one
-	if (MessagesTextMarshaller->GetNumMessages() > 0)
-	{
-		MessagesTextBox->ScrollTo(FTextLocation(MessagesTextMarshaller->GetNumMessages() - 1));
-	}
 }
 END_SLATE_FUNCTION_BUILD_OPTIMIZATION
 
@@ -708,6 +707,16 @@ SOutputLog::~SOutputLog()
 	}
 }
 
+void SOutputLog::Tick( const FGeometry& AllottedGeometry, const double InCurrentTime, const float InDeltaTime )
+{
+	if(bPendingForceScroll && MessagesTextMarshaller->GetNumMessages() > 0)
+	{
+		MessagesTextBox->ScrollTo(FTextLocation(MessagesTextMarshaller->GetNumMessages() - 1));
+	}
+
+	bPendingForceScroll = false;
+}
+
 bool SOutputLog::CreateLogMessages( const TCHAR* V, ELogVerbosity::Type Verbosity, const class FName& Category, TArray< TSharedPtr<FLogMessage> >& OutMessages )
 {
 	if (Verbosity == ELogVerbosity::SetColor)
@@ -765,7 +774,8 @@ void SOutputLog::Serialize( const TCHAR* V, ELogVerbosity::Type Verbosity, const
 	if ( MessagesTextMarshaller->AppendMessage(V, Verbosity, Category) )
 	{
 		// Don't scroll to the bottom automatically when the user is scrolling the view or has scrolled it away from the bottom.
-		if( MessagesTextBox->GetVScrollBar()->DistanceFromBottom() <= 0.f )
+		const float DistanceFromBottom = MessagesTextBox->GetVScrollBar()->DistanceFromBottom();
+		if( DistanceFromBottom <= KINDA_SMALL_NUMBER )
 		{
 			// Force a refresh so that the message has been added before we try and jump to it
 			//MessagesTextBox->Refresh();
@@ -803,3 +813,8 @@ bool SOutputLog::CanClearLog() const
 {
 	return MessagesTextMarshaller->GetNumMessages() > 0;
 }
+
+void SOutputLog::OnConsoleCommandExecuted()
+{
+	RequestForceScroll();
+}
diff --git a/Engine/Source/Developer/OutputLog/Private/SOutputLog.h b/Engine/Source/Developer/OutputLog/Private/SOutputLog.h
--- a/Engine/Source/Developer/OutputLog/Private/SOutputLog.h
+++ b/Engine/Source/Developer/OutputLog/Private/SOutputLog.h
@@ -144,6 +144,8 @@ public:
 	 */
 	void Construct( const FArguments& InArgs );
 
+	virtual void Tick( const FGeometry& AllottedGeometry, const double InCurrentTime, const float InDeltaTime ) override;
+
 	/**
 	 * Creates FLogMessage objects from FOutputDevice log callback
 	 *
@@ -156,6 +158,12 @@ public:
 	 */
 	static bool CreateLogMessages( const TCHAR* V, ELogVerbosity::Type Verbosity, const class FName& Category, TArray< TSharedPtr<FLogMessage> >& OutMessages );
 
+	/** Request we force scroll to the bottom of the log on the next Tick() */
+	void RequestForceScroll()
+	{
+		bPendingForceScroll = true;
+	}
+
 protected:
 
 	virtual void Serialize( const TCHAR* V, ELogVerbosity::Type Verbosity, const class FName& Category ) override;
@@ -176,9 +184,15 @@ private:
 	 */
 	bool CanClearLog() const;
 
+	/** Called when a console command is entered for this output log */
+	void OnConsoleCommandExecuted();
+
 	/** Converts the array of messages into something the text box understands */
 	TSharedPtr< FOutputLogTextLayoutMarshaller > MessagesTextMarshaller;
 
 	/** The editable text showing all log messages */
 	TSharedPtr< SMultiLineEditableTextBox > MessagesTextBox;
+
+	/** Are we pending a force scroll to the bottom of the log? */
+	bool bPendingForceScroll;
 };
diff --git a/Engine/Source/Editor/EditorStyle/Private/SlateEditorStyle.cpp b/Engine/Source/Editor/EditorStyle/Private/SlateEditorStyle.cpp
--- a/Engine/Source/Editor/EditorStyle/Private/SlateEditorStyle.cpp
+++ b/Engine/Source/Editor/EditorStyle/Private/SlateEditorStyle.cpp
@@ -886,8 +886,7 @@ void FSlateEditorStyle::FStyle::SetupGeneralStyles()
 		const FTextBlockStyle NormalLogText = FTextBlockStyle(NormalText)
 			.SetFont( TTF_FONT( "Fonts/DroidSansMono", 9 ) )
 			.SetColorAndOpacity( FLinearColor(FColor(0xffaaaaaa)) )
-			//.SetSelectedBackgroundColor( SelectionColor );
-			.SetSelectedBackgroundColor( FLinearColor(FColor(0xffeeeeee)) );
+			.SetSelectedBackgroundColor( FLinearColor(FColor(0xff666666)) );
 
 		Set("Log.Normal", NormalLogText );
 
diff --git a/Engine/Source/Runtime/Slate/Private/Widgets/Input/SMultiLineEditableTextBox.cpp b/Engine/Source/Runtime/Slate/Private/Widgets/Input/SMultiLineEditableTextBox.cpp
--- a/Engine/Source/Runtime/Slate/Private/Widgets/Input/SMultiLineEditableTextBox.cpp
+++ b/Engine/Source/Runtime/Slate/Private/Widgets/Input/SMultiLineEditableTextBox.cpp
@@ -33,6 +33,7 @@ void SMultiLineEditableTextBox::Construct( const FArguments& InArgs )
 		HScrollBar = SNew(SScrollBar)
 			.Style(&InArgs._Style->ScrollBarStyle)
 			.Orientation(Orient_Horizontal)
+			.AlwaysShowScrollbar(InArgs._AlwaysShowScrollbars)
 			.Thickness(FVector2D(5.0f, 5.0f));
 		HScrollBarWidget = HScrollBar;
 	}
@@ -50,6 +51,7 @@ void SMultiLineEditableTextBox::Construct( const FArguments& InArgs )
 		VScrollBar = SNew(SScrollBar)
 			.Style(&InArgs._Style->ScrollBarStyle)
 			.Orientation(Orient_Vertical)
+			.AlwaysShowScrollbar(InArgs._AlwaysShowScrollbars)
 			.Thickness(FVector2D(5.0f, 5.0f));
 		VScrollBarWidget = VScrollBar;
 	}
diff --git a/Engine/Source/Runtime/Slate/Private/Widgets/Layout/SScrollBar.cpp b/Engine/Source/Runtime/Slate/Private/Widgets/Layout/SScrollBar.cpp
--- a/Engine/Source/Runtime/Slate/Private/Widgets/Layout/SScrollBar.cpp
+++ b/Engine/Source/Runtime/Slate/Private/Widgets/Layout/SScrollBar.cpp
@@ -111,35 +111,36 @@ public:
 	{
 		const float Width = AllottedGeometry.Size.X;
 		const float Height = AllottedGeometry.Size.Y;
+		
+		// We only need to show all three children when the thumb is visible, otherwise we only need to show the track
+		if ( IsNeeded() )
+		{
+			FTrackSizeInfo TrackSizeInfo = this->GetTrackSizeInfo(AllottedGeometry);
 
-		FTrackSizeInfo TrackSizeInfo = this->GetTrackSizeInfo(AllottedGeometry);
+			// Arrange top half of the track
+			FVector2D ChildSize = (Orientation == Orient_Horizontal)
+				? FVector2D(TrackSizeInfo.ThumbStart, Height)
+				: FVector2D(Width, TrackSizeInfo.ThumbStart);
 
-		// Arrange top half of the track
-		FVector2D ChildSize = (Orientation == Orient_Horizontal)
-			? FVector2D(TrackSizeInfo.ThumbStart, Height)
-			: FVector2D(Width, TrackSizeInfo.ThumbStart);
+			FVector2D ChildPos(0, 0);
+			ArrangedChildren.AddWidget(			
+				AllottedGeometry.MakeChild( Children[TOP_SLOT_INDEX].GetWidget(), ChildPos, ChildSize )
+			);
 
-		FVector2D ChildPos(0, 0);
-		ArrangedChildren.AddWidget(			
-			AllottedGeometry.MakeChild( Children[TOP_SLOT_INDEX].GetWidget(), ChildPos, ChildSize )
-		);
+			// Arrange bottom half of the track
+			ChildPos = (Orientation == Orient_Horizontal)
+				? FVector2D(TrackSizeInfo.GetThumbEnd(), 0)
+				: FVector2D(0, TrackSizeInfo.GetThumbEnd());
 
-		// Arrange bottom half of the track
-		ChildPos = (Orientation == Orient_Horizontal)
-			? FVector2D(TrackSizeInfo.GetThumbEnd(), 0)
-			: FVector2D(0, TrackSizeInfo.GetThumbEnd());
+			ChildSize = (Orientation == Orient_Horizontal)
+				? FVector2D(AllottedGeometry.Size.X - TrackSizeInfo.GetThumbEnd(), Height)
+				: FVector2D(Width, AllottedGeometry.Size.Y - TrackSizeInfo.GetThumbEnd());
 
-		ChildSize = (Orientation == Orient_Horizontal)
-			? FVector2D(AllottedGeometry.Size.X - TrackSizeInfo.GetThumbEnd(), Height)
-			: FVector2D(Width, AllottedGeometry.Size.Y - TrackSizeInfo.GetThumbEnd());
+			ArrangedChildren.AddWidget( 
+				AllottedGeometry.MakeChild( Children[BOTTOM_SLOT_INDEX].GetWidget(), ChildPos, ChildSize )
+			);
 
-		ArrangedChildren.AddWidget( 
-			AllottedGeometry.MakeChild( Children[BOTTOM_SLOT_INDEX].GetWidget(), ChildPos, ChildSize )
-		);
-		
-		// Arrange the thumb; only show it in enabled scrollbars
-		if ( IsNeeded() )
-		{
+			// Arrange the thumb
 			ChildPos = (Orientation == Orient_Horizontal)
 				? FVector2D(TrackSizeInfo.ThumbStart, 0)
 				: FVector2D(0, TrackSizeInfo.ThumbStart);
@@ -152,6 +153,13 @@ public:
 				AllottedGeometry.MakeChild( Children[THUMB_SLOT_INDEX].GetWidget(), ChildPos, ChildSize )
 			);
 		}
+		else
+		{
+			// No thumb is visible, so just show the top half of the track at the current width/height
+			ArrangedChildren.AddWidget(			
+				AllottedGeometry.MakeChild( Children[TOP_SLOT_INDEX].GetWidget(), FVector2D(0, 0), FVector2D(Width, Height) )
+			);
+		}
 
 	}
 	
diff --git a/Engine/Source/Runtime/Slate/Public/Widgets/Input/SMultiLineEditableTextBox.h b/Engine/Source/Runtime/Slate/Public/Widgets/Input/SMultiLineEditableTextBox.h
--- a/Engine/Source/Runtime/Slate/Public/Widgets/Input/SMultiLineEditableTextBox.h
+++ b/Engine/Source/Runtime/Slate/Public/Widgets/Input/SMultiLineEditableTextBox.h
@@ -29,6 +29,7 @@ public:
 		, _SelectAllTextWhenFocused( false )
 		, _RevertTextOnEscape( false )
 		, _ClearKeyboardFocusOnCommit( true )
+		, _AlwaysShowScrollbars( false )
 		, _WrapTextAt(0.0f)
 		, _AutoWrapText(false)
 		, _SelectAllTextOnCommit( false )
@@ -86,6 +87,9 @@ public:
 		/** Whether to clear keyboard focus when pressing enter to commit changes */
 		SLATE_ATTRIBUTE( bool, ClearKeyboardFocusOnCommit )
 
+		/** Should we always show the scrollbars (only affects internally created scroll bars) */
+		SLATE_ARGUMENT(bool, AlwaysShowScrollbars)
+
 		/** The horizontal scroll bar widget, or null to create one internally */
 		SLATE_ARGUMENT( TSharedPtr< SScrollBar >, HScrollBar )
 
