Removed Tick() from SOutputLog

It wasn't actually needed as we're able to perform the scroll immediately in all cases we need to, rather than wait until the next Tick().

diff --git a/Engine/Source/Developer/OutputLog/Private/SOutputLog.cpp b/Engine/Source/Developer/OutputLog/Private/SOutputLog.cpp
--- a/Engine/Source/Developer/OutputLog/Private/SOutputLog.cpp
+++ b/Engine/Source/Developer/OutputLog/Private/SOutputLog.cpp
@@ -659,9 +659,6 @@ FOutputLogTextLayoutMarshaller::FOutputLogTextLayoutMarshaller(TArray< TSharedPt
 BEGIN_SLATE_FUNCTION_BUILD_OPTIMIZATION
 void SOutputLog::Construct( const FArguments& InArgs )
 {
-	// Force scroll in the first Tick() once we've been sized and populated correctly
-	bPendingForceScroll = true;
-
 	MessagesTextMarshaller = FOutputLogTextLayoutMarshaller::Create(MoveTemp(InArgs._Messages));
 
 	MessagesTextBox = SNew(SMultiLineEditableTextBox)
@@ -696,6 +693,8 @@ void SOutputLog::Construct( const FArguments& InArgs )
 	];
 	
 	GLog->AddOutputDevice(this);
+
+	RequestForceScroll();
 }
 END_SLATE_FUNCTION_BUILD_OPTIMIZATION
 
@@ -707,16 +706,6 @@ SOutputLog::~SOutputLog()
 	}
 }
 
-void SOutputLog::Tick( const FGeometry& AllottedGeometry, const double InCurrentTime, const float InDeltaTime )
-{
-	if(bPendingForceScroll && MessagesTextMarshaller->GetNumMessages() > 0)
-	{
-		MessagesTextBox->ScrollTo(FTextLocation(MessagesTextMarshaller->GetNumMessages() - 1));
-	}
-
-	bPendingForceScroll = false;
-}
-
 bool SOutputLog::CreateLogMessages( const TCHAR* V, ELogVerbosity::Type Verbosity, const class FName& Category, TArray< TSharedPtr<FLogMessage> >& OutMessages )
 {
 	if (Verbosity == ELogVerbosity::SetColor)
@@ -818,3 +807,11 @@ void SOutputLog::OnConsoleCommandExecuted()
 {
 	RequestForceScroll();
 }
+
+void SOutputLog::RequestForceScroll()
+{
+	if(MessagesTextMarshaller->GetNumMessages() > 0)
+	{
+		MessagesTextBox->ScrollTo(FTextLocation(MessagesTextMarshaller->GetNumMessages() - 1));
+	}
+}
diff --git a/Engine/Source/Developer/OutputLog/Private/SOutputLog.h b/Engine/Source/Developer/OutputLog/Private/SOutputLog.h
--- a/Engine/Source/Developer/OutputLog/Private/SOutputLog.h
+++ b/Engine/Source/Developer/OutputLog/Private/SOutputLog.h
@@ -144,8 +144,6 @@ public:
 	 */
 	void Construct( const FArguments& InArgs );
 
-	virtual void Tick( const FGeometry& AllottedGeometry, const double InCurrentTime, const float InDeltaTime ) override;
-
 	/**
 	 * Creates FLogMessage objects from FOutputDevice log callback
 	 *
@@ -158,12 +156,6 @@ public:
 	 */
 	static bool CreateLogMessages( const TCHAR* V, ELogVerbosity::Type Verbosity, const class FName& Category, TArray< TSharedPtr<FLogMessage> >& OutMessages );
 
-	/** Request we force scroll to the bottom of the log on the next Tick() */
-	void RequestForceScroll()
-	{
-		bPendingForceScroll = true;
-	}
-
 protected:
 
 	virtual void Serialize( const TCHAR* V, ELogVerbosity::Type Verbosity, const class FName& Category ) override;
@@ -187,12 +179,12 @@ private:
 	/** Called when a console command is entered for this output log */
 	void OnConsoleCommandExecuted();
 
+	/** Request we immediately force scroll to the bottom of the log */
+	void RequestForceScroll();
+
 	/** Converts the array of messages into something the text box understands */
 	TSharedPtr< FOutputLogTextLayoutMarshaller > MessagesTextMarshaller;
 
 	/** The editable text showing all log messages */
 	TSharedPtr< SMultiLineEditableTextBox > MessagesTextBox;
-
-	/** Are we pending a force scroll to the bottom of the log? */
-	bool bPendingForceScroll;
 };
