Fix crash typing in the output log when a PIE window is open

diff --git a/Engine/Source/Developer/OutputLog/Private/SOutputLog.cpp b/Engine/Source/Developer/OutputLog/Private/SOutputLog.cpp
--- a/Engine/Source/Developer/OutputLog/Private/SOutputLog.cpp
+++ b/Engine/Source/Developer/OutputLog/Private/SOutputLog.cpp
@@ -322,7 +322,7 @@ void SConsoleInputBox::OnTextCommitted( const FText& InText, ETextCommit::Type C
 				UWorld* OldWorld = NULL;
 
 				// The play world needs to handle these commands if it exists
-				if( GIsEditor && GEditor->PlayWorld )
+				if( GIsEditor && GEditor->PlayWorld && !GIsPlayInEditorWorld )
 				{
 					World = GEditor->PlayWorld;
 					OldWorld = SetPlayInEditorWorld( GEditor->PlayWorld );
