Refactor to unify keyboard and controller input and focus.
Native controller navigation support.

diff --git a/Engine/Config/BaseEngine.ini b/Engine/Config/BaseEngine.ini
--- a/Engine/Config/BaseEngine.ini
+++ b/Engine/Config/BaseEngine.ini
@@ -611,6 +611,8 @@ bUseStreamingPause=false
 +EnumRedirects=(EnumName="EPhysicsType", OldEnumEntry="PhysType_Unfixed",NewEnumEntry="PhysType_Simulated")
 ;Struct Redirects for UE4
 +ActiveStructRedirects=(OldStructName="VimDebugData",NewStructName="AnimBlueprintDebugData")
++ActiveStructRedirects=(OldStructName="KeyboardEvent",NewStructName="KeyEvent")
++ActiveStructRedirects=(OldStructName="KeyboardFocusEvent",NewStructName="FocusEvent")
 
 ; Property name redirects
 +TaggedPropertyRedirects=(ClassName="Actor",OldPropertyName="LifeSpan",NewPropertyName="InitialLifeSpan")
diff --git a/Engine/Plugins/2D/Paper2D/Source/Paper2DEditor/Private/SPaperEditorViewport.cpp b/Engine/Plugins/2D/Paper2D/Source/Paper2DEditor/Private/SPaperEditorViewport.cpp
--- a/Engine/Plugins/2D/Paper2D/Source/Paper2DEditor/Private/SPaperEditorViewport.cpp
+++ b/Engine/Plugins/2D/Paper2D/Source/Paper2DEditor/Private/SPaperEditorViewport.cpp
@@ -267,7 +267,7 @@ FReply SPaperEditorViewport::OnMouseButtonUp(const FGeometry& MyGeometry, const
 
 		this->bIsPanning = false;
 		return (WidgetToFocus.IsValid())
-			? ReplyState.ReleaseMouseCapture().SetKeyboardFocus( WidgetToFocus.ToSharedRef(), EKeyboardFocusCause::SetDirectly )
+			? ReplyState.ReleaseMouseCapture().SetUserFocus(WidgetToFocus.ToSharedRef(), EFocusCause::SetDirectly)
 			: ReplyState.ReleaseMouseCapture();
 	}
 	else if (MouseEvent.GetEffectingButton() == EKeys::LeftMouseButton)
diff --git a/Engine/Plugins/2D/Paper2D/Source/Paper2DEditor/Private/TileMapEditing/STileLayerItem.cpp b/Engine/Plugins/2D/Paper2D/Source/Paper2DEditor/Private/TileMapEditing/STileLayerItem.cpp
--- a/Engine/Plugins/2D/Paper2D/Source/Paper2DEditor/Private/TileMapEditing/STileLayerItem.cpp
+++ b/Engine/Plugins/2D/Paper2D/Source/Paper2DEditor/Private/TileMapEditing/STileLayerItem.cpp
@@ -84,7 +84,7 @@ FReply SRenamableEntry::OnMouseButtonDoubleClick(const FGeometry& InMyGeometry,
 
 		FWidgetPath TextBoxPath;
 		FSlateApplication::Get().GeneratePathToWidgetUnchecked(EditableTextWidget.ToSharedRef(), TextBoxPath);
-		FSlateApplication::Get().SetKeyboardFocus(EditableTextWidget, EKeyboardFocusCause::Mouse);
+		FSlateApplication::Get().SetKeyboardFocus(EditableTextWidget, EFocusCause::Mouse);
 
 		return FReply::Handled();
 	}
diff --git a/Engine/Source/Developer/AutomationWindow/Private/SAutomationWindow.cpp b/Engine/Source/Developer/AutomationWindow/Private/SAutomationWindow.cpp
--- a/Engine/Source/Developer/AutomationWindow/Private/SAutomationWindow.cpp
+++ b/Engine/Source/Developer/AutomationWindow/Private/SAutomationWindow.cpp
@@ -885,7 +885,7 @@ void SAutomationWindow::ExpandEnabledTests( TSharedPtr< IAutomationReport > InRe
 FReply SAutomationWindow::HandleNewPresetClicked()
 {
 	bAddingTestPreset = true;
-	return FReply::Handled().SetKeyboardFocus(PresetTextBox.ToSharedRef(), EKeyboardFocusCause::SetDirectly);
+	return FReply::Handled().SetUserFocus(PresetTextBox.ToSharedRef(), EFocusCause::SetDirectly);
 }
 
 FReply SAutomationWindow::HandleSavePresetClicked()
@@ -1905,9 +1905,9 @@ bool SAutomationWindow::IsAnySelectedRowEnabled()
 /* SWidget implementation
  *****************************************************************************/
 
-FReply SAutomationWindow::OnKeyUp( const FGeometry& InGeometry, const FKeyboardEvent& InKeyboardEvent )
+FReply SAutomationWindow::OnKeyUp( const FGeometry& InGeometry, const FKeyEvent& InKeyEvent )
 {
-	if (InKeyboardEvent.GetKey() == EKeys::SpaceBar)
+	if (InKeyEvent.GetKey() == EKeys::SpaceBar)
 	{
 		SetAllSelectedTestsChecked(IsAnySelectedRowEnabled());
 		return FReply::Handled();
@@ -1916,11 +1916,11 @@ FReply SAutomationWindow::OnKeyUp( const FGeometry& InGeometry, const FKeyboardE
 }
 
 
-FReply SAutomationWindow::OnKeyDown( const FGeometry& MyGeometry, const FKeyboardEvent& InKeyboardEvent )
+FReply SAutomationWindow::OnKeyDown( const FGeometry& MyGeometry, const FKeyEvent& InKeyEvent )
 {
-	if (InKeyboardEvent.IsControlDown())
+	if (InKeyEvent.IsControlDown())
 	{
-		if (InKeyboardEvent.GetKey() == EKeys::C)
+		if (InKeyEvent.GetKey() == EKeys::C)
 		{
 			CopyLog();
 
diff --git a/Engine/Source/Developer/AutomationWindow/Private/SAutomationWindow.h b/Engine/Source/Developer/AutomationWindow/Private/SAutomationWindow.h
--- a/Engine/Source/Developer/AutomationWindow/Private/SAutomationWindow.h
+++ b/Engine/Source/Developer/AutomationWindow/Private/SAutomationWindow.h
@@ -121,10 +121,10 @@ protected:
 	bool IsAnySelectedRowEnabled();
 
 	/** Overridden from SWidget: Called when a key is pressed down - capturing copy. */
-	virtual FReply OnKeyDown( const FGeometry& MyGeometry, const FKeyboardEvent& InKeyboardEvent ) override;
+	virtual FReply OnKeyDown( const FGeometry& MyGeometry, const FKeyEvent& InKeyEvent ) override;
 
 	/** Overridden from SWidget: Called after a key is released. */
-	virtual FReply OnKeyUp( const FGeometry& InGeometry, const FKeyboardEvent& InKeyboardEvent ) override;
+	virtual FReply OnKeyUp( const FGeometry& InGeometry, const FKeyEvent& InKeyEvent ) override;
 
 private:
 
diff --git a/Engine/Source/Developer/LogVisualizer/Private/SLogVisualizer.cpp b/Engine/Source/Developer/LogVisualizer/Private/SLogVisualizer.cpp
--- a/Engine/Source/Developer/LogVisualizer/Private/SLogVisualizer.cpp
+++ b/Engine/Source/Developer/LogVisualizer/Private/SLogVisualizer.cpp
@@ -1186,13 +1186,13 @@ FReply SLogVisualizer::OnMouseWheel( const FGeometry& MyGeometry, const FPointer
 	return SCompoundWidget::OnMouseWheel(MyGeometry, MouseEvent);
 }
 
-FReply SLogVisualizer::OnKeyDown( const FGeometry& MyGeometry, const FKeyboardEvent& InKeyboardEvent )
+FReply SLogVisualizer::OnKeyDown( const FGeometry& MyGeometry, const FKeyEvent& InKeyEvent )
 {
-	const FKey Key = InKeyboardEvent.GetKey();
+	const FKey Key = InKeyEvent.GetKey();
 	if (Key == EKeys::Left || Key == EKeys::Right)
 	{
 		int32 MoveBy = Key == EKeys::Left ? -1 : 1;
-		if (InKeyboardEvent.IsLeftControlDown())
+		if (InKeyEvent.IsLeftControlDown())
 		{
 			MoveBy *= 10;
 		}
@@ -1202,7 +1202,7 @@ FReply SLogVisualizer::OnKeyDown( const FGeometry& MyGeometry, const FKeyboardEv
 		return FReply::Handled();
 	}
 
-	return SCompoundWidget::OnKeyDown(MyGeometry, InKeyboardEvent);
+	return SCompoundWidget::OnKeyDown(MyGeometry, InKeyEvent);
 }
 
 TSharedRef< SWidget > SLogVisualizer::MakeMainMenu()
diff --git a/Engine/Source/Developer/LogVisualizer/Private/SLogVisualizer.h b/Engine/Source/Developer/LogVisualizer/Private/SLogVisualizer.h
--- a/Engine/Source/Developer/LogVisualizer/Private/SLogVisualizer.h
+++ b/Engine/Source/Developer/LogVisualizer/Private/SLogVisualizer.h
@@ -88,7 +88,7 @@ public:
 
 	// overrides
 	virtual FReply OnMouseWheel( const FGeometry& MyGeometry, const FPointerEvent& MouseEvent ) override;
-	virtual FReply OnKeyDown( const FGeometry& MyGeometry, const FKeyboardEvent& InKeyboardEvent) override;
+	virtual FReply OnKeyDown( const FGeometry& MyGeometry, const FKeyEvent& InKeyEvent) override;
 
 	// Get delegates
 	const FSlateBrush* GetRecordButtonBrush() const;
diff --git a/Engine/Source/Developer/MessageLog/Private/UserInterface/SMessageLog.cpp b/Engine/Source/Developer/MessageLog/Private/UserInterface/SMessageLog.cpp
--- a/Engine/Source/Developer/MessageLog/Private/UserInterface/SMessageLog.cpp
+++ b/Engine/Source/Developer/MessageLog/Private/UserInterface/SMessageLog.cpp
@@ -107,11 +107,11 @@ void SMessageLog::HandleSelectionUpdated()
 /* SWidget overrides
  *****************************************************************************/
 
-FReply SMessageLog::OnKeyDown( const FGeometry& MyGeometry, const FKeyboardEvent& InKeyboardEvent )
+FReply SMessageLog::OnKeyDown( const FGeometry& MyGeometry, const FKeyEvent& InKeyEvent )
 {
 	if (LogListing.IsValid())
 	{
-		return LogListing->GetCommandList()->ProcessCommandBindings(InKeyboardEvent)
+		return LogListing->GetCommandList()->ProcessCommandBindings(InKeyEvent)
 			? FReply::Handled()
 			: FReply::Unhandled();
 	}
diff --git a/Engine/Source/Developer/MessageLog/Private/UserInterface/SMessageLog.h b/Engine/Source/Developer/MessageLog/Private/UserInterface/SMessageLog.h
--- a/Engine/Source/Developer/MessageLog/Private/UserInterface/SMessageLog.h
+++ b/Engine/Source/Developer/MessageLog/Private/UserInterface/SMessageLog.h
@@ -32,7 +32,7 @@ public:
 
 	// SWidget overrides
 
-	FReply OnKeyDown( const FGeometry& MyGeometry, const FKeyboardEvent& InKeyboardEvent ) override;
+	FReply OnKeyDown( const FGeometry& MyGeometry, const FKeyEvent& InKeyEvent ) override;
 
 private:
 
diff --git a/Engine/Source/Developer/MessageLog/Private/UserInterface/SMessageLogListing.cpp b/Engine/Source/Developer/MessageLog/Private/UserInterface/SMessageLogListing.cpp
--- a/Engine/Source/Developer/MessageLog/Private/UserInterface/SMessageLogListing.cpp
+++ b/Engine/Source/Developer/MessageLog/Private/UserInterface/SMessageLogListing.cpp
@@ -307,9 +307,9 @@ const TSharedRef< const FUICommandList > SMessageLogListing::GetCommandList() co
 }
 
 
-FReply SMessageLogListing::OnKeyDown( const FGeometry& MyGeometry, const FKeyboardEvent& InKeyboardEvent )
+FReply SMessageLogListing::OnKeyDown( const FGeometry& MyGeometry, const FKeyEvent& InKeyEvent )
 {
-	return UICommandList->ProcessCommandBindings( InKeyboardEvent ) ? FReply::Handled() : FReply::Unhandled();
+	return UICommandList->ProcessCommandBindings( InKeyEvent ) ? FReply::Handled() : FReply::Unhandled();
 }
 
 
diff --git a/Engine/Source/Developer/MessageLog/Private/UserInterface/SMessageLogListing.h b/Engine/Source/Developer/MessageLog/Private/UserInterface/SMessageLogListing.h
--- a/Engine/Source/Developer/MessageLog/Private/UserInterface/SMessageLogListing.h
+++ b/Engine/Source/Developer/MessageLog/Private/UserInterface/SMessageLogListing.h
@@ -83,13 +83,13 @@ public:
 	const TSharedRef< const FUICommandList > GetCommandList() const;
 
 	/**
-	 * Called after a key is pressed when this widget has keyboard focus
+	 * Called after a key is pressed when this widget has focus
 	 *
-	 * @param  InKeyboardEvent  Keyboard event
+	 * @param  InKeyEvent  Key event
 	 *
 	 * @return  Returns whether the event was handled, along with other possible actions
 	 */
-	FReply OnKeyDown( const FGeometry& MyGeometry, const FKeyboardEvent& InKeyboardEvent ) override;
+	FReply OnKeyDown( const FGeometry& MyGeometry, const FKeyEvent& InKeyEvent ) override;
 
 	/** Delegate supplying a label for the page-flipper widget */
 	FString GetPageText() const;
diff --git a/Engine/Source/Developer/OutputLog/Private/SDebugConsole.cpp b/Engine/Source/Developer/OutputLog/Private/SDebugConsole.cpp
--- a/Engine/Source/Developer/OutputLog/Private/SDebugConsole.cpp
+++ b/Engine/Source/Developer/OutputLog/Private/SDebugConsole.cpp
@@ -103,7 +103,7 @@ SDebugConsole::SDebugConsole()
 
 void SDebugConsole::SetFocusToEditableText()
 {
-	FSlateApplication::Get().SetKeyboardFocus( EditableTextBox.ToSharedRef(), EKeyboardFocusCause::SetDirectly );
+	FSlateApplication::Get().SetKeyboardFocus( EditableTextBox.ToSharedRef(), EFocusCause::SetDirectly );
 }
 
 EVisibility SDebugConsole::MakeVisibleIfLogIsShown() const
diff --git a/Engine/Source/Developer/OutputLog/Private/SOutputLog.cpp b/Engine/Source/Developer/OutputLog/Private/SOutputLog.cpp
--- a/Engine/Source/Developer/OutputLog/Private/SOutputLog.cpp
+++ b/Engine/Source/Developer/OutputLog/Private/SOutputLog.cpp
@@ -68,16 +68,16 @@ private:
 			);
 		}
 
-		virtual FReply OnKeyDown( const FGeometry& MyGeometry, const FKeyboardEvent& InKeyboardEvent )
+		virtual FReply OnKeyDown( const FGeometry& MyGeometry, const FKeyEvent& InKeyEvent )
 		{
 			// Special case handling.  Intercept the tilde key.  It is not suitable for typing in the console
-			if( InKeyboardEvent.GetKey() == EKeys::Tilde )
+			if( InKeyEvent.GetKey() == EKeys::Tilde )
 			{
 				return FReply::Unhandled();
 			}
 			else
 			{
-				return SEditableText::OnKeyDown( MyGeometry, InKeyboardEvent );
+				return SEditableText::OnKeyDown( MyGeometry, InKeyEvent );
 			}
 		}
 
@@ -200,7 +200,7 @@ void SConsoleInputBox::SuggestionSelectionChanged(TSharedPtr<FString> NewValue,
 			// Ideally this would set the focus back to the edit control
 //			FWidgetPath WidgetToFocusPath;
 //			FSlateApplication::Get().GeneratePathToWidgetUnchecked( InputText.ToSharedRef(), WidgetToFocusPath );
-//			FSlateApplication::Get().SetKeyboardFocus( WidgetToFocusPath, EKeyboardFocusCause::SetDirectly );
+//			FSlateApplication::Get().SetKeyboardFocus( WidgetToFocusPath, EFocusCause::SetDirectly );
 			break;
 		}
 	}
@@ -384,13 +384,13 @@ void SConsoleInputBox::OnTextCommitted( const FText& InText, ETextCommit::Type C
 	}
 }
 
-FReply SConsoleInputBox::OnKeyDown( const FGeometry& MyGeometry, const FKeyboardEvent& KeyboardEvent )
+FReply SConsoleInputBox::OnKeyDown( const FGeometry& MyGeometry, const FKeyEvent& KeyEvent )
 {
 	if(SuggestionBox->IsOpen())
 	{
-		if(KeyboardEvent.GetKey() == EKeys::Up || KeyboardEvent.GetKey() == EKeys::Down)
+		if(KeyEvent.GetKey() == EKeys::Up || KeyEvent.GetKey() == EKeys::Down)
 		{
-			if(KeyboardEvent.GetKey() == EKeys::Up)
+			if(KeyEvent.GetKey() == EKeys::Up)
 			{
 				if(SelectedSuggestion < 0)
 				{
@@ -404,7 +404,7 @@ FReply SConsoleInputBox::OnKeyDown( const FGeometry& MyGeometry, const FKeyboard
 				}
 			}
 
-			if(KeyboardEvent.GetKey() == EKeys::Down)
+			if(KeyEvent.GetKey() == EKeys::Down)
 			{
 				if(SelectedSuggestion < Suggestions.Num() - 1)
 				{
@@ -422,7 +422,7 @@ FReply SConsoleInputBox::OnKeyDown( const FGeometry& MyGeometry, const FKeyboard
 
 			return FReply::Handled();
 		}
-		else if (KeyboardEvent.GetKey() == EKeys::Tab)
+		else if (KeyEvent.GetKey() == EKeys::Tab)
 		{
 			if (Suggestions.Num())
 			{
@@ -443,7 +443,7 @@ FReply SConsoleInputBox::OnKeyDown( const FGeometry& MyGeometry, const FKeyboard
 	}
 	else
 	{
-		if(KeyboardEvent.GetKey() == EKeys::Up)
+		if(KeyEvent.GetKey() == EKeys::Up)
 		{
 			TArray<FString> History;
 
@@ -495,7 +495,7 @@ void SConsoleInputBox::SetSuggestions(TArray<FString>& Elements, bool bInHistory
 		// Force the textbox back into focus.
 		FWidgetPath WidgetToFocusPath;
 		FSlateApplication::Get().GeneratePathToWidgetUnchecked( InputText.ToSharedRef(), WidgetToFocusPath );
-		FSlateApplication::Get().SetKeyboardFocus( WidgetToFocusPath, EKeyboardFocusCause::SetDirectly );
+		FSlateApplication::Get().SetKeyboardFocus( WidgetToFocusPath, EFocusCause::SetDirectly );
 	}
 	else
 	{
@@ -503,7 +503,7 @@ void SConsoleInputBox::SetSuggestions(TArray<FString>& Elements, bool bInHistory
 	}
 }
 
-void SConsoleInputBox::OnKeyboardFocusLost( const FKeyboardFocusEvent& InKeyboardFocusEvent )
+void SConsoleInputBox::OnFocusLost( const FFocusEvent& InFocusEvent )
 {
 //	SuggestionBox->SetIsOpen(false);
 }
diff --git a/Engine/Source/Developer/OutputLog/Private/SOutputLog.h b/Engine/Source/Developer/OutputLog/Private/SOutputLog.h
--- a/Engine/Source/Developer/OutputLog/Private/SOutputLog.h
+++ b/Engine/Source/Developer/OutputLog/Private/SOutputLog.h
@@ -67,9 +67,9 @@ protected:
 	virtual bool SupportsKeyboardFocus() const { return true; }
 
 	// e.g. Tab or Key_Up
-	virtual FReply OnKeyDown( const FGeometry& MyGeometry, const FKeyboardEvent& KeyboardEvent );
+	virtual FReply OnKeyDown( const FGeometry& MyGeometry, const FKeyEvent& KeyEvent );
 
-	void OnKeyboardFocusLost( const FKeyboardFocusEvent& InKeyboardFocusEvent );
+	void OnFocusLost( const FFocusEvent& InFocusEvent );
 
 	/** Handles entering in a command */
 	void OnTextCommitted(const FText& InText, ETextCommit::Type CommitInfo);
diff --git a/Engine/Source/Developer/OutputLog/Public/OutputLogModule.h b/Engine/Source/Developer/OutputLog/Public/OutputLogModule.h
--- a/Engine/Source/Developer/OutputLog/Public/OutputLogModule.h
+++ b/Engine/Source/Developer/OutputLog/Public/OutputLogModule.h
@@ -19,7 +19,7 @@ namespace EDebugConsoleStyle
 
 struct FDebugConsoleDelegates
 {
-	FSimpleDelegate OnKeyboardFocusLost;
+	FSimpleDelegate OnFocusLost;
 	FSimpleDelegate OnConsoleCommandExecuted;
 };
 
diff --git a/Engine/Source/Developer/Profiler/Private/Widgets/SProfilerWindow.cpp b/Engine/Source/Developer/Profiler/Private/Widgets/SProfilerWindow.cpp
--- a/Engine/Source/Developer/Profiler/Private/Widgets/SProfilerWindow.cpp
+++ b/Engine/Source/Developer/Profiler/Private/Widgets/SProfilerWindow.cpp
@@ -515,9 +515,9 @@ void SProfilerWindow::OnMouseLeave(const FPointerEvent& MouseEvent)
 	bIsActive = false;
 }
 
-FReply SProfilerWindow::OnKeyDown( const FGeometry& MyGeometry, const FKeyboardEvent& InKeyboardEvent )
+FReply SProfilerWindow::OnKeyDown( const FGeometry& MyGeometry, const FKeyEvent& InKeyEvent )
 {
-	return FProfilerManager::Get()->GetCommandList()->ProcessCommandBindings( InKeyboardEvent ) ? FReply::Handled() : FReply::Unhandled();
+	return FProfilerManager::Get()->GetCommandList()->ProcessCommandBindings( InKeyEvent ) ? FReply::Handled() : FReply::Unhandled();
 }
 
 FReply SProfilerWindow::OnDragOver( const FGeometry& MyGeometry, const FDragDropEvent& DragDropEvent )
diff --git a/Engine/Source/Developer/Profiler/Private/Widgets/SProfilerWindow.h b/Engine/Source/Developer/Profiler/Private/Widgets/SProfilerWindow.h
--- a/Engine/Source/Developer/Profiler/Private/Widgets/SProfilerWindow.h
+++ b/Engine/Source/Developer/Profiler/Private/Widgets/SProfilerWindow.h
@@ -86,14 +86,14 @@ private:
 	virtual void OnMouseLeave(const FPointerEvent& MouseEvent) override;
 
 	/**
-	 * Called after a key is pressed when this widget has keyboard focus
+	 * Called after a key is pressed when this widget has focus
 	 *
 	 * @param MyGeometry The Geometry of the widget receiving the event
-	 * @param  InKeyboardEvent  Keyboard event
+	 * @param  InKeyEvent  Key event
 	 *
 	 * @return  Returns whether the event was handled, along with other possible actions
 	 */
-	virtual FReply OnKeyDown( const FGeometry& MyGeometry, const FKeyboardEvent& InKeyboardEvent ) override;
+	virtual FReply OnKeyDown( const FGeometry& MyGeometry, const FKeyEvent& InKeyEvent ) override;
 
 	/**
 	 * Called when the user is dropping something onto a widget; terminates drag and drop.
diff --git a/Engine/Source/Developer/SessionFrontend/Private/Widgets/Console/SSessionConsole.cpp b/Engine/Source/Developer/SessionFrontend/Private/Widgets/Console/SSessionConsole.cpp
--- a/Engine/Source/Developer/SessionFrontend/Private/Widgets/Console/SSessionConsole.cpp
+++ b/Engine/Source/Developer/SessionFrontend/Private/Widgets/Console/SSessionConsole.cpp
@@ -348,17 +348,17 @@ void SSessionConsole::SendCommand( const FString& CommandString )
 /* SWidget implementation
  *****************************************************************************/
 
-FReply SSessionConsole::OnKeyDown( const FGeometry& MyGeometry, const FKeyboardEvent& InKeyboardEvent )
+FReply SSessionConsole::OnKeyDown( const FGeometry& MyGeometry, const FKeyEvent& InKeyEvent )
 {
-	if (InKeyboardEvent.IsControlDown())
+	if (InKeyEvent.IsControlDown())
 	{
-		if (InKeyboardEvent.GetKey() == EKeys::C)
+		if (InKeyEvent.GetKey() == EKeys::C)
 		{
 			CopyLog();
 
 			return FReply::Handled();
 		}
-		else if (InKeyboardEvent.GetKey() == EKeys::S)
+		else if (InKeyEvent.GetKey() == EKeys::S)
 		{
 			SaveLog();
 
diff --git a/Engine/Source/Developer/SessionFrontend/Private/Widgets/Console/SSessionConsole.h b/Engine/Source/Developer/SessionFrontend/Private/Widgets/Console/SSessionConsole.h
--- a/Engine/Source/Developer/SessionFrontend/Private/Widgets/Console/SSessionConsole.h
+++ b/Engine/Source/Developer/SessionFrontend/Private/Widgets/Console/SSessionConsole.h
@@ -62,7 +62,7 @@ protected:
 
 protected:
 
-	virtual FReply OnKeyDown( const FGeometry& MyGeometry, const FKeyboardEvent& InKeyboardEvent ) override;
+	virtual FReply OnKeyDown( const FGeometry& MyGeometry, const FKeyEvent& InKeyEvent ) override;
 
 private:
 
diff --git a/Engine/Source/Developer/SettingsEditor/Private/Widgets/SSettingsEditor.cpp b/Engine/Source/Developer/SettingsEditor/Private/Widgets/SSettingsEditor.cpp
--- a/Engine/Source/Developer/SettingsEditor/Private/Widgets/SSettingsEditor.cpp
+++ b/Engine/Source/Developer/SettingsEditor/Private/Widgets/SSettingsEditor.cpp
@@ -667,7 +667,7 @@ void SSettingsEditor::HandleModelSelectionChanged()
 
 		if (FSlateApplication::Get().GeneratePathToWidgetUnchecked(FocusWidget.ToSharedRef(), FocusWidgetPath))
 		{
-			FSlateApplication::Get().SetKeyboardFocus(FocusWidgetPath, EKeyboardFocusCause::SetDirectly);
+			FSlateApplication::Get().SetKeyboardFocus(FocusWidgetPath, EFocusCause::SetDirectly);
 		}
 	}
 	else
diff --git a/Engine/Source/Developer/SuperSearch/Private/SSuperSearch.cpp b/Engine/Source/Developer/SuperSearch/Private/SSuperSearch.cpp
--- a/Engine/Source/Developer/SuperSearch/Private/SSuperSearch.cpp
+++ b/Engine/Source/Developer/SuperSearch/Private/SSuperSearch.cpp
@@ -359,13 +359,13 @@ void SSuperSearchBox::Query_HttpRequestComplete(FHttpRequestPtr HttpRequest, FHt
 }
 
 
-FReply SSuperSearchBox::OnKeyDown( const FGeometry& MyGeometry, const FKeyboardEvent& KeyboardEvent )
+FReply SSuperSearchBox::OnKeyDown( const FGeometry& MyGeometry, const FKeyEvent& KeyEvent )
 {
 	if(SuggestionBox->IsOpen() && Suggestions.Num())
 	{
-		if(KeyboardEvent.GetKey() == EKeys::Up || KeyboardEvent.GetKey() == EKeys::Down)
+		if(KeyEvent.GetKey() == EKeys::Up || KeyEvent.GetKey() == EKeys::Down)
 		{
-			if(KeyboardEvent.GetKey() == EKeys::Up)
+			if(KeyEvent.GetKey() == EKeys::Up)
 			{
 				//if we're at the top swing around
 				if(SelectedSuggestion == 1)
@@ -386,7 +386,7 @@ FReply SSuperSearchBox::OnKeyDown( const FGeometry& MyGeometry, const FKeyboardE
 				}
 			}
 
-			if(KeyboardEvent.GetKey() == EKeys::Down)
+			if(KeyEvent.GetKey() == EKeys::Down)
 			{
 				if(SelectedSuggestion < Suggestions.Num() - 1)
 				{
@@ -478,10 +478,10 @@ void SSuperSearchBox::UpdateSuggestions()
 	MarkActiveSuggestion();
 
 	// Force the textbox back into focus.
-	FSlateApplication::Get().SetKeyboardFocus(InputText.ToSharedRef(), EKeyboardFocusCause::SetDirectly);
+	FSlateApplication::Get().SetKeyboardFocus(InputText.ToSharedRef(), EFocusCause::SetDirectly);
 }
 
-void SSuperSearchBox::OnKeyboardFocusLost( const FKeyboardFocusEvent& InKeyboardFocusEvent )
+void SSuperSearchBox::OnFocusLost( const FFocusEvent& InFocusEvent )
 {
 //	SuggestionBox->SetIsOpen(false);
 }
diff --git a/Engine/Source/Developer/SuperSearch/Private/SSuperSearch.h b/Engine/Source/Developer/SuperSearch/Private/SSuperSearch.h
--- a/Engine/Source/Developer/SuperSearch/Private/SSuperSearch.h
+++ b/Engine/Source/Developer/SuperSearch/Private/SSuperSearch.h
@@ -53,9 +53,9 @@ protected:
 	virtual bool SupportsKeyboardFocus() const { return true; }
 
 	// e.g. Tab or Key_Up
-	virtual FReply OnKeyDown( const FGeometry& MyGeometry, const FKeyboardEvent& KeyboardEvent );
+	virtual FReply OnKeyDown( const FGeometry& MyGeometry, const FKeyEvent& KeyEvent );
 
-	void OnKeyboardFocusLost( const FKeyboardFocusEvent& InKeyboardFocusEvent );
+	void OnFocusLost( const FFocusEvent& InFocusEvent );
 
 	/** Handles entering in a command */
 	void OnTextCommitted(const FText& InText, ETextCommit::Type CommitInfo);
diff --git a/Engine/Source/Developer/WebBrowser/Private/WebBrowserViewport.cpp b/Engine/Source/Developer/WebBrowser/Private/WebBrowserViewport.cpp
--- a/Engine/Source/Developer/WebBrowser/Private/WebBrowserViewport.cpp
+++ b/Engine/Source/Developer/WebBrowser/Private/WebBrowserViewport.cpp
@@ -80,15 +80,15 @@ FReply FWebBrowserViewport::OnMouseButtonDoubleClick(const FGeometry& InMyGeomet
 	return FReply::Handled();
 }
 
-FReply FWebBrowserViewport::OnKeyDown(const FGeometry& MyGeometry, const FKeyboardEvent& InKeyboardEvent)
+FReply FWebBrowserViewport::OnKeyDown(const FGeometry& MyGeometry, const FKeyEvent& InKeyEvent)
 {
-	WebBrowserWindow->OnKeyDown(InKeyboardEvent);
+	WebBrowserWindow->OnKeyDown(InKeyEvent);
 	return FReply::Handled();
 }
 
-FReply FWebBrowserViewport::OnKeyUp(const FGeometry& MyGeometry, const FKeyboardEvent& InKeyboardEvent)
+FReply FWebBrowserViewport::OnKeyUp(const FGeometry& MyGeometry, const FKeyEvent& InKeyEvent)
 {
-	WebBrowserWindow->OnKeyUp(InKeyboardEvent);
+	WebBrowserWindow->OnKeyUp(InKeyEvent);
 	return FReply::Handled();
 }
 
@@ -98,13 +98,13 @@ FReply FWebBrowserViewport::OnKeyChar( const FGeometry& MyGeometry, const FChara
 	return FReply::Handled();
 }
 
-FReply FWebBrowserViewport::OnKeyboardFocusReceived(const FKeyboardFocusEvent& InKeyboardFocusEvent)
+FReply FWebBrowserViewport::OnFocusReceived(const FFocusEvent& InFocusEvent)
 {
 	WebBrowserWindow->OnFocus(true);
 	return FReply::Handled();
 }
 
-void FWebBrowserViewport::OnKeyboardFocusLost(const FKeyboardFocusEvent& InKeyboardFocusEvent)
+void FWebBrowserViewport::OnFocusLost(const FFocusEvent& InFocusEvent)
 {
 	WebBrowserWindow->OnFocus(false);
 }
diff --git a/Engine/Source/Developer/WebBrowser/Private/WebBrowserWindow.cpp b/Engine/Source/Developer/WebBrowser/Private/WebBrowserWindow.cpp
--- a/Engine/Source/Developer/WebBrowser/Private/WebBrowserWindow.cpp
+++ b/Engine/Source/Developer/WebBrowser/Private/WebBrowserWindow.cpp
@@ -86,37 +86,37 @@ FString FWebBrowserWindow::GetTitle() const
 	return Title;
 }
 
-void FWebBrowserWindow::OnKeyDown(const FKeyboardEvent& InKeyboardEvent)
+void FWebBrowserWindow::OnKeyDown(const FKeyEvent& InKeyEvent)
 {
 	if (IsValid())
 	{
 		CefKeyEvent KeyEvent;
-		KeyEvent.windows_key_code = InKeyboardEvent.GetKeyCode();
+		KeyEvent.windows_key_code = InKeyEvent.GetKeyCode();
 		// TODO: Figure out whether this is a system key if we come across problems
 		/*KeyEvent.is_system_key = message == WM_SYSCHAR ||
 			message == WM_SYSKEYDOWN ||
 			message == WM_SYSKEYUP;*/
 
 		KeyEvent.type = KEYEVENT_RAWKEYDOWN;
-		KeyEvent.modifiers = GetCefKeyboardModifiers(InKeyboardEvent);
+		KeyEvent.modifiers = GetCefKeyboardModifiers(InKeyEvent);
 
 		InternalCefBrowser->GetHost()->SendKeyEvent(KeyEvent);
 	}
 }
 
-void FWebBrowserWindow::OnKeyUp(const FKeyboardEvent& InKeyboardEvent)
+void FWebBrowserWindow::OnKeyUp(const FKeyEvent& InKeyEvent)
 {
 	if (IsValid())
 	{
 		CefKeyEvent KeyEvent;
-		KeyEvent.windows_key_code = InKeyboardEvent.GetKeyCode();
+		KeyEvent.windows_key_code = InKeyEvent.GetKeyCode();
 		// TODO: Figure out whether this is a system key if we come across problems
 		/*KeyEvent.is_system_key = message == WM_SYSCHAR ||
 			message == WM_SYSKEYDOWN ||
 			message == WM_SYSKEYUP;*/
 
 		KeyEvent.type = KEYEVENT_KEYUP;
-		KeyEvent.modifiers = GetCefKeyboardModifiers(InKeyboardEvent);
+		KeyEvent.modifiers = GetCefKeyboardModifiers(InKeyEvent);
 
 		InternalCefBrowser->GetHost()->SendKeyEvent(KeyEvent);
 	}
@@ -377,11 +377,11 @@ void FWebBrowserWindow::OnCursorChange(CefCursorHandle Cursor)
 	//::SetCursor( Cursor );
 }
 
-int32 FWebBrowserWindow::GetCefKeyboardModifiers(const FKeyboardEvent& KeyboardEvent)
+int32 FWebBrowserWindow::GetCefKeyboardModifiers(const FKeyEvent& KeyEvent)
 {
-	int32 Modifiers = GetCefInputModifiers(KeyboardEvent);
+	int32 Modifiers = GetCefInputModifiers(KeyEvent);
 
-	const FKey Key = KeyboardEvent.GetKey();
+	const FKey Key = KeyEvent.GetKey();
 	if (Key == EKeys::LeftAlt ||
 		Key == EKeys::LeftCommand ||
 		Key == EKeys::LeftControl ||
diff --git a/Engine/Source/Developer/WebBrowser/Private/WebBrowserWindow.h b/Engine/Source/Developer/WebBrowser/Private/WebBrowserWindow.h
--- a/Engine/Source/Developer/WebBrowser/Private/WebBrowserWindow.h
+++ b/Engine/Source/Developer/WebBrowser/Private/WebBrowserWindow.h
@@ -45,8 +45,8 @@ public:
 	virtual bool IsValid() const override;
 	virtual bool IsClosing() const override;
 	virtual FString GetTitle() const override;
-	virtual void OnKeyDown(const FKeyboardEvent& InKeyboardEvent) override;
-	virtual void OnKeyUp(const FKeyboardEvent& InKeyboardEvent) override;
+	virtual void OnKeyDown(const FKeyEvent& InKeyEvent) override;
+	virtual void OnKeyUp(const FKeyEvent& InKeyEvent) override;
 	virtual void OnKeyChar(const FCharacterEvent& InCharacterEvent) override;
 	virtual void OnMouseButtonDown(const FGeometry& MyGeometry, const FPointerEvent& MouseEvent) override;
 	virtual void OnMouseButtonUp(const FGeometry& MyGeometry, const FPointerEvent& MouseEvent) override;
@@ -118,12 +118,12 @@ private:
 	void OnCursorChange(CefCursorHandle Cursor);
 
 	/**
-	 * Gets the Cef Keyboard Modifiers based on a Keyboard Event
+	 * Gets the Cef Keyboard Modifiers based on a Key Event
 	 * 
-	 * @param KeyboardEvent The Keyboard event
+	 * @param KeyEvent The Key event
 	 * @return Bits representing keyboard modifiers
 	 */
-	static int32 GetCefKeyboardModifiers(const FKeyboardEvent& KeyboardEvent);
+	static int32 GetCefKeyboardModifiers(const FKeyEvent& KeyEvent);
 
 	/**
 	 * Gets the Cef Mouse Modifiers based on a Mouse Event
diff --git a/Engine/Source/Developer/WebBrowser/Public/IWebBrowserWindow.h b/Engine/Source/Developer/WebBrowser/Public/IWebBrowserWindow.h
--- a/Engine/Source/Developer/WebBrowser/Public/IWebBrowserWindow.h
+++ b/Engine/Source/Developer/WebBrowser/Public/IWebBrowserWindow.h
@@ -4,7 +4,7 @@
 
 // Forward Declarations
 struct FGeometry;
-struct FKeyboardEvent;
+struct FKeyEvent;
 struct FCharacterEvent;
 struct FPointerEvent;
 class FSlateShaderResource;
@@ -49,16 +49,16 @@ public:
 	/**
 	 * Notify the browser that a key has been pressed
 	 *
-	 * @param  InKeyboardEvent  Keyboard event
+	 * @param  InKeyEvent  Key event
 	 */
-	virtual void OnKeyDown(const FKeyboardEvent& InKeyboardEvent) = 0;
+	virtual void OnKeyDown(const FKeyEvent& InKeyEvent) = 0;
 
 	/**
 	 * Notify the browser that a key has been released
 	 *
-	 * @param  InKeyboardEvent  Keyboard event
+	 * @param  InKeyEvent  Key event
 	 */
-	virtual void OnKeyUp(const FKeyboardEvent& InKeyboardEvent) = 0;
+	virtual void OnKeyUp(const FKeyEvent& InKeyEvent) = 0;
 
 	/**
 	 * Notify the browser of a character event
diff --git a/Engine/Source/Developer/WebBrowser/Public/WebBrowserViewport.h b/Engine/Source/Developer/WebBrowser/Public/WebBrowserViewport.h
--- a/Engine/Source/Developer/WebBrowser/Public/WebBrowserViewport.h
+++ b/Engine/Source/Developer/WebBrowser/Public/WebBrowserViewport.h
@@ -44,11 +44,11 @@ public:
 	virtual FReply OnMouseMove( const FGeometry& MyGeometry, const FPointerEvent& MouseEvent ) override;
 	virtual FReply OnMouseWheel( const FGeometry& MyGeometry, const FPointerEvent& MouseEvent ) override;
 	virtual FReply OnMouseButtonDoubleClick( const FGeometry& InMyGeometry, const FPointerEvent& InMouseEvent ) override;
-	virtual FReply OnKeyDown( const FGeometry& MyGeometry, const FKeyboardEvent& InKeyboardEvent ) override;
-	virtual FReply OnKeyUp( const FGeometry& MyGeometry, const FKeyboardEvent& InKeyboardEvent ) override;
+	virtual FReply OnKeyDown( const FGeometry& MyGeometry, const FKeyEvent& InKeyEvent ) override;
+	virtual FReply OnKeyUp( const FGeometry& MyGeometry, const FKeyEvent& InKeyEvent ) override;
 	virtual FReply OnKeyChar( const FGeometry& MyGeometry, const FCharacterEvent& InCharacterEvent ) override;
-	virtual FReply OnKeyboardFocusReceived( const FKeyboardFocusEvent& InKeyboardFocusEvent ) override;
-	virtual void OnKeyboardFocusLost( const FKeyboardFocusEvent& InKeyboardFocusEvent ) override;
+	virtual FReply OnFocusReceived( const FFocusEvent& InFocusEvent ) override;
+	virtual void OnFocusLost( const FFocusEvent& InFocusEvent ) override;
 
 private:
 	/** The web browser this viewport will display */
diff --git a/Engine/Source/Editor/BehaviorTreeEditor/Private/FindInBT.cpp b/Engine/Source/Editor/BehaviorTreeEditor/Private/FindInBT.cpp
--- a/Engine/Source/Editor/BehaviorTreeEditor/Private/FindInBT.cpp
+++ b/Engine/Source/Editor/BehaviorTreeEditor/Private/FindInBT.cpp
@@ -150,7 +150,7 @@ void SFindInBT::FocusForUse()
 	FSlateApplication::Get().GeneratePathToWidgetUnchecked(SearchTextField.ToSharedRef(), FilterTextBoxWidgetPath);
 
 	// Set keyboard focus directly
-	FSlateApplication::Get().SetKeyboardFocus(FilterTextBoxWidgetPath, EKeyboardFocusCause::SetDirectly);
+	FSlateApplication::Get().SetKeyboardFocus(FilterTextBoxWidgetPath, EFocusCause::SetDirectly);
 }
 
 void SFindInBT::OnSearchTextChanged(const FText& Text)
diff --git a/Engine/Source/Editor/BspMode/Private/SBspPalette.cpp b/Engine/Source/Editor/BspMode/Private/SBspPalette.cpp
--- a/Engine/Source/Editor/BspMode/Private/SBspPalette.cpp
+++ b/Engine/Source/Editor/BspMode/Private/SBspPalette.cpp
@@ -14,7 +14,7 @@ class SBspBuilderListView : public SListView<TSharedPtr<FBspBuilderType>>
 {
 public:
 	virtual bool SupportsKeyboardFocus() const override { return true; }
-	virtual FReply OnKeyDown( const FGeometry& InGeometry, const FKeyboardEvent& InKeyboardEvent ) override
+	virtual FReply OnKeyDown( const FGeometry& InGeometry, const FKeyEvent& InKeyEvent ) override
 	{
 		return FReply::Unhandled();
 	}
diff --git a/Engine/Source/Editor/Cascade/Private/Cascade.cpp b/Engine/Source/Editor/Cascade/Private/Cascade.cpp
--- a/Engine/Source/Editor/Cascade/Private/Cascade.cpp
+++ b/Engine/Source/Editor/Cascade/Private/Cascade.cpp
@@ -1896,7 +1896,7 @@ void FCascade::NotifyPostChange(const FPropertyChangedEvent& PropertyChangedEven
 
 void FCascade::PreEditCurve(TArray<UObject*> CurvesAboutToChange)
 {
-	FSlateApplication::Get().ClearKeyboardFocus(EKeyboardFocusCause::Mouse);
+	FSlateApplication::Get().ClearKeyboardFocus(EFocusCause::Mouse);
 
 	//Need to keep text in-sync with PostEditCurve
 	BeginTransaction( NSLOCTEXT("UnrealEd", "EditCurve", "Edit Curve") );
diff --git a/Engine/Source/Editor/ClassViewer/Private/SClassViewer.cpp b/Engine/Source/Editor/ClassViewer/Private/SClassViewer.cpp
--- a/Engine/Source/Editor/ClassViewer/Private/SClassViewer.cpp
+++ b/Engine/Source/Editor/ClassViewer/Private/SClassViewer.cpp
@@ -2884,14 +2884,14 @@ void SClassViewer::Populate()
 	}
 }
 
-FReply SClassViewer::OnKeyDown( const FGeometry& MyGeometry, const FKeyboardEvent& InKeyboardEvent )
+FReply SClassViewer::OnKeyDown( const FGeometry& MyGeometry, const FKeyEvent& InKeyEvent )
 {
 	// Forward key down to class tree
-	return ClassTree->OnKeyDown(MyGeometry,InKeyboardEvent);
+	return ClassTree->OnKeyDown(MyGeometry,InKeyEvent);
 }
 
 
-FReply SClassViewer::OnKeyboardFocusReceived( const FGeometry& MyGeometry, const FKeyboardFocusEvent& InKeyboardFocusEvent )
+FReply SClassViewer::OnFocusReceived( const FGeometry& MyGeometry, const FFocusEvent& InFocusEvent )
 {
 	if (RootTreeItems.Num() > 0)
 	{
@@ -2900,7 +2900,7 @@ FReply SClassViewer::OnKeyboardFocusReceived( const FGeometry& MyGeometry, const
 		OnClassViewerSelectionChanged(RootTreeItems[0],ESelectInfo::OnMouseClick);
 	}
 
-	FSlateApplication::Get().SetKeyboardFocus(SearchBox.ToSharedRef(), EKeyboardFocusCause::SetDirectly);
+	FSlateApplication::Get().SetKeyboardFocus(SearchBox.ToSharedRef(), EFocusCause::SetDirectly);
 	
 	return FReply::Unhandled();
 }
@@ -2937,7 +2937,7 @@ void SClassViewer::Tick( const FGeometry& AllottedGeometry, const double InCurre
 	{
 		FWidgetPath WidgetToFocusPath;
 		FSlateApplication::Get().GeneratePathToWidgetUnchecked( SearchBox.ToSharedRef(), WidgetToFocusPath );
-		FSlateApplication::Get().SetKeyboardFocus( WidgetToFocusPath, EKeyboardFocusCause::SetDirectly );
+		FSlateApplication::Get().SetKeyboardFocus( WidgetToFocusPath, EFocusCause::SetDirectly );
 		bPendingFocusNextFrame = false;
 	}
 
diff --git a/Engine/Source/Editor/ClassViewer/Private/SClassViewer.h b/Engine/Source/Editor/ClassViewer/Private/SClassViewer.h
--- a/Engine/Source/Editor/ClassViewer/Private/SClassViewer.h
+++ b/Engine/Source/Editor/ClassViewer/Private/SClassViewer.h
@@ -38,8 +38,8 @@ public:
 
 	/** SWidget interface */
 	virtual void Tick( const FGeometry& AllottedGeometry, const double InCurrentTime, const float InDeltaTime ) override;
-	virtual FReply OnKeyDown( const FGeometry& MyGeometry, const FKeyboardEvent& InKeyboardEvent ) override;
-	virtual FReply OnKeyboardFocusReceived( const FGeometry& MyGeometry, const FKeyboardFocusEvent& InKeyboardFocusEvent ) override;
+	virtual FReply OnKeyDown( const FGeometry& MyGeometry, const FKeyEvent& InKeyEvent ) override;
+	virtual FReply OnFocusReceived( const FGeometry& MyGeometry, const FFocusEvent& InFocusEvent ) override;
 	virtual bool SupportsKeyboardFocus() const override;
 
 private:
diff --git a/Engine/Source/Editor/ContentBrowser/Private/AssetViewWidgets.cpp b/Engine/Source/Editor/ContentBrowser/Private/AssetViewWidgets.cpp
--- a/Engine/Source/Editor/ContentBrowser/Private/AssetViewWidgets.cpp
+++ b/Engine/Source/Editor/ContentBrowser/Private/AssetViewWidgets.cpp
@@ -19,10 +19,10 @@
 // FAssetViewModeUtils
 ///////////////////////////////
 
-FReply FAssetViewModeUtils::OnViewModeKeyDown( const TSet< TSharedPtr<FAssetViewItem> >& SelectedItems, const FKeyboardEvent& InKeyboardEvent )
+FReply FAssetViewModeUtils::OnViewModeKeyDown( const TSet< TSharedPtr<FAssetViewItem> >& SelectedItems, const FKeyEvent& InKeyEvent )
 {
 	// All asset views use Ctrl-C to copy references to assets
-	if ( InKeyboardEvent.IsControlDown() && InKeyboardEvent.GetCharacter() == 'C' )
+	if ( InKeyEvent.IsControlDown() && InKeyEvent.GetCharacter() == 'C' )
 	{
 		FAssetRegistryModule& AssetRegistryModule = FModuleManager::LoadModuleChecked<FAssetRegistryModule>(TEXT("AssetRegistry"));
 
@@ -61,9 +61,9 @@ FReply FAssetViewModeUtils::OnViewModeKeyDown( const TSet< TSharedPtr<FAssetView
 // Asset view modes
 ///////////////////////////////
 
-FReply SAssetTileView::OnKeyDown( const FGeometry& InGeometry, const FKeyboardEvent& InKeyboardEvent )
+FReply SAssetTileView::OnKeyDown( const FGeometry& InGeometry, const FKeyEvent& InKeyEvent )
 {
-	FReply Reply = FAssetViewModeUtils::OnViewModeKeyDown(SelectedItems, InKeyboardEvent);
+	FReply Reply = FAssetViewModeUtils::OnViewModeKeyDown(SelectedItems, InKeyEvent);
 
 	if ( Reply.IsEventHandled() )
 	{
@@ -71,13 +71,13 @@ FReply SAssetTileView::OnKeyDown( const FGeometry& InGeometry, const FKeyboardEv
 	}
 	else
 	{
-		return STileView<TSharedPtr<FAssetViewItem>>::OnKeyDown(InGeometry, InKeyboardEvent);
+		return STileView<TSharedPtr<FAssetViewItem>>::OnKeyDown(InGeometry, InKeyEvent);
 	}
 }
 
-FReply SAssetListView::OnKeyDown( const FGeometry& InGeometry, const FKeyboardEvent& InKeyboardEvent )
+FReply SAssetListView::OnKeyDown( const FGeometry& InGeometry, const FKeyEvent& InKeyEvent )
 {
-	FReply Reply = FAssetViewModeUtils::OnViewModeKeyDown(SelectedItems, InKeyboardEvent);
+	FReply Reply = FAssetViewModeUtils::OnViewModeKeyDown(SelectedItems, InKeyEvent);
 
 	if ( Reply.IsEventHandled() )
 	{
@@ -85,13 +85,13 @@ FReply SAssetListView::OnKeyDown( const FGeometry& InGeometry, const FKeyboardEv
 	}
 	else
 	{
-		return SListView<TSharedPtr<FAssetViewItem>>::OnKeyDown(InGeometry, InKeyboardEvent);
+		return SListView<TSharedPtr<FAssetViewItem>>::OnKeyDown(InGeometry, InKeyEvent);
 	}
 }
 
-FReply SAssetColumnView::OnKeyDown( const FGeometry& InGeometry, const FKeyboardEvent& InKeyboardEvent )
+FReply SAssetColumnView::OnKeyDown( const FGeometry& InGeometry, const FKeyEvent& InKeyEvent )
 {
-	FReply Reply = FAssetViewModeUtils::OnViewModeKeyDown(SelectedItems, InKeyboardEvent);
+	FReply Reply = FAssetViewModeUtils::OnViewModeKeyDown(SelectedItems, InKeyEvent);
 
 	if ( Reply.IsEventHandled() )
 	{
@@ -99,7 +99,7 @@ FReply SAssetColumnView::OnKeyDown( const FGeometry& InGeometry, const FKeyboard
 	}
 	else
 	{
-		return SListView<TSharedPtr<FAssetViewItem>>::OnKeyDown(InGeometry, InKeyboardEvent);
+		return SListView<TSharedPtr<FAssetViewItem>>::OnKeyDown(InGeometry, InKeyEvent);
 	}
 }
 
diff --git a/Engine/Source/Editor/ContentBrowser/Private/AssetViewWidgets.h b/Engine/Source/Editor/ContentBrowser/Private/AssetViewWidgets.h
--- a/Engine/Source/Editor/ContentBrowser/Private/AssetViewWidgets.h
+++ b/Engine/Source/Editor/ContentBrowser/Private/AssetViewWidgets.h
@@ -9,7 +9,7 @@ DECLARE_DELEGATE_OneParam( FOnItemDestroyed, const TSharedPtr<FAssetViewItem>& /
 
 namespace FAssetViewModeUtils 
 {
-	FReply OnViewModeKeyDown( const TSet< TSharedPtr<FAssetViewItem> >& SelectedItems, const FKeyboardEvent& InKeyboardEvent );
+	FReply OnViewModeKeyDown( const TSet< TSharedPtr<FAssetViewItem> >& SelectedItems, const FKeyEvent& InKeyEvent );
 }
 
 /** The tile view mode of the asset view */
@@ -17,7 +17,7 @@ class SAssetTileView : public STileView<TSharedPtr<FAssetViewItem>>
 {
 public:
 	virtual bool SupportsKeyboardFocus() const override { return true; }
-	virtual FReply OnKeyDown( const FGeometry& InGeometry, const FKeyboardEvent& InKeyboardEvent ) override;
+	virtual FReply OnKeyDown( const FGeometry& InGeometry, const FKeyEvent& InKeyEvent ) override;
 };
 
 /** The list view mode of the asset view */
@@ -25,7 +25,7 @@ class SAssetListView : public SListView<TSharedPtr<FAssetViewItem>>
 {
 public:
 	virtual bool SupportsKeyboardFocus() const override { return true; }
-	virtual FReply OnKeyDown( const FGeometry& InGeometry, const FKeyboardEvent& InKeyboardEvent ) override;
+	virtual FReply OnKeyDown( const FGeometry& InGeometry, const FKeyEvent& InKeyEvent ) override;
 };
 
 /** The columns view mode of the asset view */
@@ -33,7 +33,7 @@ class SAssetColumnView : public SListView<TSharedPtr<FAssetViewItem>>
 {
 public:
 	virtual bool SupportsKeyboardFocus() const override { return true; }
-	virtual FReply OnKeyDown( const FGeometry& InGeometry, const FKeyboardEvent& InKeyboardEvent ) override;
+	virtual FReply OnKeyDown( const FGeometry& InGeometry, const FKeyEvent& InKeyEvent ) override;
 };
 
 /** A base class for all asset view items */
diff --git a/Engine/Source/Editor/ContentBrowser/Private/SAssetDialog.cpp b/Engine/Source/Editor/ContentBrowser/Private/SAssetDialog.cpp
--- a/Engine/Source/Editor/ContentBrowser/Private/SAssetDialog.cpp
+++ b/Engine/Source/Editor/ContentBrowser/Private/SAssetDialog.cpp
@@ -196,15 +196,15 @@ void SAssetDialog::Construct(const FArguments& InArgs, const FSharedAssetDialogC
 }
 END_SLATE_FUNCTION_BUILD_OPTIMIZATION
 
-FReply SAssetDialog::OnKeyDown( const FGeometry& MyGeometry, const FKeyboardEvent& InKeyboardEvent )
+FReply SAssetDialog::OnKeyDown( const FGeometry& MyGeometry, const FKeyEvent& InKeyEvent )
 {
-	if ( InKeyboardEvent.GetKey() == EKeys::Escape )
+	if ( InKeyEvent.GetKey() == EKeys::Escape )
 	{
 		CloseDialog();
 		return FReply::Handled();
 	}
 
-	return SCompoundWidget::OnKeyDown(MyGeometry, InKeyboardEvent);
+	return SCompoundWidget::OnKeyDown(MyGeometry, InKeyEvent);
 }
 
 void SAssetDialog::Tick(const FGeometry& AllottedGeometry, const double InCurrentTime, const float InDeltaTime)
@@ -232,7 +232,7 @@ void SAssetDialog::FocusNameBox()
 {
 	if ( NameEditableText.IsValid() )
 	{
-		FSlateApplication::Get().SetKeyboardFocus(NameEditableText.ToSharedRef(), EKeyboardFocusCause::SetDirectly);
+		FSlateApplication::Get().SetKeyboardFocus(NameEditableText.ToSharedRef(), EFocusCause::SetDirectly);
 	}
 }
 
diff --git a/Engine/Source/Editor/ContentBrowser/Private/SAssetDialog.h b/Engine/Source/Editor/ContentBrowser/Private/SAssetDialog.h
--- a/Engine/Source/Editor/ContentBrowser/Private/SAssetDialog.h
+++ b/Engine/Source/Editor/ContentBrowser/Private/SAssetDialog.h
@@ -13,7 +13,7 @@ public:
 
 	virtual void Construct(const FArguments& InArgs, const FSharedAssetDialogConfig& InConfig);
 
-	virtual FReply OnKeyDown( const FGeometry& MyGeometry, const FKeyboardEvent& InKeyboardEvent ) override;
+	virtual FReply OnKeyDown( const FGeometry& MyGeometry, const FKeyEvent& InKeyEvent ) override;
 	void Tick(const FGeometry& AllottedGeometry, const double InCurrentTime, const float InDeltaTime) override;
 
 	/** Sets the delegate handler for when an open operation is committed */
diff --git a/Engine/Source/Editor/ContentBrowser/Private/SAssetPicker.cpp b/Engine/Source/Editor/ContentBrowser/Private/SAssetPicker.cpp
--- a/Engine/Source/Editor/ContentBrowser/Private/SAssetPicker.cpp
+++ b/Engine/Source/Editor/ContentBrowser/Private/SAssetPicker.cpp
@@ -236,27 +236,27 @@ void SAssetPicker::Tick(const FGeometry& AllottedGeometry, const double InCurren
 	{
 		FWidgetPath WidgetToFocusPath;
 		FSlateApplication::Get().GeneratePathToWidgetUnchecked( SearchBoxPtr.ToSharedRef(), WidgetToFocusPath );
-		FSlateApplication::Get().SetKeyboardFocus( WidgetToFocusPath, EKeyboardFocusCause::SetDirectly );
+		FSlateApplication::Get().SetKeyboardFocus( WidgetToFocusPath, EFocusCause::SetDirectly );
 		bPendingFocusNextFrame = false;
 	}
 
 	SCompoundWidget::Tick(AllottedGeometry, InCurrentTime, InDeltaTime);
 }
 
-FReply SAssetPicker::OnKeyDown(const FGeometry& MyGeometry, const FKeyboardEvent& InKeyboardEvent)
+FReply SAssetPicker::OnKeyDown(const FGeometry& MyGeometry, const FKeyEvent& InKeyEvent)
 {
 	// Up and down move thru the filtered list
 	int32 SelectionDelta = 0;
 
-	if (InKeyboardEvent.GetKey() == EKeys::Up)
+	if (InKeyEvent.GetKey() == EKeys::Up)
 	{
 		SelectionDelta = -1;
 	}
-	else if (InKeyboardEvent.GetKey() == EKeys::Down)
+	else if (InKeyEvent.GetKey() == EKeys::Down)
 	{
 		SelectionDelta = +1;
 	}
-	else if (InKeyboardEvent.GetKey() == EKeys::Enter)
+	else if (InKeyEvent.GetKey() == EKeys::Enter)
 	{
 		TArray<FAssetData> SelectionSet = AssetViewPtr->GetSelectedAssets();
 		HandleAssetsActivated(SelectionSet, EAssetTypeActivationMethod::Opened);
@@ -271,7 +271,7 @@ FReply SAssetPicker::OnKeyDown(const FGeometry& MyGeometry, const FKeyboardEvent
 		return FReply::Handled();
 	}
 
-	if (Commands->ProcessCommandBindings(InKeyboardEvent))
+	if (Commands->ProcessCommandBindings(InKeyEvent))
 	{
 		return FReply::Handled();
 	}
diff --git a/Engine/Source/Editor/ContentBrowser/Private/SAssetView.cpp b/Engine/Source/Editor/ContentBrowser/Private/SAssetView.cpp
--- a/Engine/Source/Editor/ContentBrowser/Private/SAssetView.cpp
+++ b/Engine/Source/Editor/ContentBrowser/Private/SAssetView.cpp
@@ -1164,13 +1164,13 @@ FReply SAssetView::OnKeyChar( const FGeometry& MyGeometry,const FCharacterEvent&
 	return FReply::Unhandled();
 }
 
-FReply SAssetView::OnKeyDown( const FGeometry& MyGeometry, const FKeyboardEvent& InKeyboardEvent )
+FReply SAssetView::OnKeyDown( const FGeometry& MyGeometry, const FKeyEvent& InKeyEvent )
 {
 	{
 		// Swallow the key-presses used by the quick-jump in OnKeyChar to avoid other things (such as the viewport commands) getting them instead
 		// eg) Pressing "W" without this would set the viewport to "translate" mode
 		const bool bTestOnly = true;
-		if(HandleQuickJumpKeyDown(InKeyboardEvent.GetCharacter(), InKeyboardEvent.IsControlDown(), InKeyboardEvent.IsAltDown(), bTestOnly).IsEventHandled())
+		if(HandleQuickJumpKeyDown(InKeyEvent.GetCharacter(), InKeyEvent.IsControlDown(), InKeyEvent.IsAltDown(), bTestOnly).IsEventHandled())
 		{
 			return FReply::Handled();
 		}
@@ -1193,7 +1193,7 @@ FReply SAssetView::OnMouseWheel( const FGeometry& MyGeometry, const FPointerEven
 	return FReply::Unhandled();
 }
 
-void SAssetView::OnKeyboardFocusChanging( const FWeakWidgetPath& PreviousFocusPath, const FWidgetPath& NewWidgetPath )
+void SAssetView::OnFocusChanging( const FWeakWidgetPath& PreviousFocusPath, const FWidgetPath& NewWidgetPath )
 {
 	ResetQuickJump();
 }
@@ -2492,9 +2492,9 @@ void SAssetView::FocusList() const
 {
 	switch ( GetCurrentViewType() )
 	{
-		case EAssetViewType::List: FSlateApplication::Get().SetKeyboardFocus(ListView, EKeyboardFocusCause::SetDirectly); break;
-		case EAssetViewType::Tile: FSlateApplication::Get().SetKeyboardFocus(TileView, EKeyboardFocusCause::SetDirectly); break;
-		case EAssetViewType::Column: FSlateApplication::Get().SetKeyboardFocus(ColumnView, EKeyboardFocusCause::SetDirectly); break;
+		case EAssetViewType::List: FSlateApplication::Get().SetKeyboardFocus(ListView, EFocusCause::SetDirectly); break;
+		case EAssetViewType::Tile: FSlateApplication::Get().SetKeyboardFocus(TileView, EFocusCause::SetDirectly); break;
+		case EAssetViewType::Column: FSlateApplication::Get().SetKeyboardFocus(ColumnView, EFocusCause::SetDirectly); break;
 	}
 }
 
diff --git a/Engine/Source/Editor/ContentBrowser/Private/SAssetView.h b/Engine/Source/Editor/ContentBrowser/Private/SAssetView.h
--- a/Engine/Source/Editor/ContentBrowser/Private/SAssetView.h
+++ b/Engine/Source/Editor/ContentBrowser/Private/SAssetView.h
@@ -214,9 +214,9 @@ public:
 	virtual FReply OnDragOver( const FGeometry& MyGeometry, const FDragDropEvent& DragDropEvent ) override;
 	virtual FReply OnDrop( const FGeometry& MyGeometry, const FDragDropEvent& DragDropEvent ) override;
 	virtual FReply OnKeyChar( const FGeometry& MyGeometry,const FCharacterEvent& InCharacterEvent ) override;
-	virtual FReply OnKeyDown( const FGeometry& MyGeometry, const FKeyboardEvent& InKeyboardEvent ) override;
+	virtual FReply OnKeyDown( const FGeometry& MyGeometry, const FKeyEvent& InKeyEvent ) override;
 	virtual FReply OnMouseWheel( const FGeometry& MyGeometry, const FPointerEvent& MouseEvent ) override;
-	virtual void OnKeyboardFocusChanging( const FWeakWidgetPath& PreviousFocusPath, const FWidgetPath& NewWidgetPath ) override;
+	virtual void OnFocusChanging( const FWeakWidgetPath& PreviousFocusPath, const FWidgetPath& NewWidgetPath ) override;
 
 	/** Opens the selected assets or folders, depending on the selection */
 	void OnOpenAssetsOrFolders();
diff --git a/Engine/Source/Editor/ContentBrowser/Private/SCollectionView.cpp b/Engine/Source/Editor/ContentBrowser/Private/SCollectionView.cpp
--- a/Engine/Source/Editor/ContentBrowser/Private/SCollectionView.cpp
+++ b/Engine/Source/Editor/ContentBrowser/Private/SCollectionView.cpp
@@ -308,9 +308,9 @@ void SCollectionView::Tick( const FGeometry& AllottedGeometry, const double InCu
 	SCompoundWidget::Tick(AllottedGeometry, InCurrentTime, InDeltaTime);
 }
 
-FReply SCollectionView::OnKeyDown( const FGeometry& MyGeometry, const FKeyboardEvent& InKeyboardEvent )
+FReply SCollectionView::OnKeyDown( const FGeometry& MyGeometry, const FKeyEvent& InKeyEvent )
 {
-	if( Commands->ProcessCommandBindings( InKeyboardEvent ) )
+	if( Commands->ProcessCommandBindings( InKeyEvent ) )
 	{
 		return FReply::Handled();
 	}
diff --git a/Engine/Source/Editor/ContentBrowser/Private/SContentBrowser.cpp b/Engine/Source/Editor/ContentBrowser/Private/SContentBrowser.cpp
--- a/Engine/Source/Editor/ContentBrowser/Private/SContentBrowser.cpp
+++ b/Engine/Source/Editor/ContentBrowser/Private/SContentBrowser.cpp
@@ -805,12 +805,12 @@ bool SContentBrowser::IsLocked() const
 void SContentBrowser::SetKeyboardFocusOnSearch() const
 {
 	// Focus on the search box
-	FSlateApplication::Get().SetKeyboardFocus( SearchBoxPtr, EKeyboardFocusCause::SetDirectly );
+	FSlateApplication::Get().SetKeyboardFocus( SearchBoxPtr, EFocusCause::SetDirectly );
 }
 
-FReply SContentBrowser::OnKeyDown( const FGeometry& MyGeometry, const FKeyboardEvent& InKeyboardEvent )
+FReply SContentBrowser::OnKeyDown( const FGeometry& MyGeometry, const FKeyEvent& InKeyEvent )
 {
-	if( Commands->ProcessCommandBindings( InKeyboardEvent ) )
+	if( Commands->ProcessCommandBindings( InKeyEvent ) )
 	{
 		return FReply::Handled();
 	}
diff --git a/Engine/Source/Editor/ContentBrowser/Private/SPathView.cpp b/Engine/Source/Editor/ContentBrowser/Private/SPathView.cpp
--- a/Engine/Source/Editor/ContentBrowser/Private/SPathView.cpp
+++ b/Engine/Source/Editor/ContentBrowser/Private/SPathView.cpp
@@ -622,7 +622,7 @@ void SPathView::Tick( const FGeometry& AllottedGeometry, const double InCurrentT
 	{
 		FWidgetPath WidgetToFocusPath;
 		FSlateApplication::Get().GeneratePathToWidgetUnchecked( SearchBoxPtr.ToSharedRef(), WidgetToFocusPath );
-		FSlateApplication::Get().SetKeyboardFocus( WidgetToFocusPath, EKeyboardFocusCause::SetDirectly );
+		FSlateApplication::Get().SetKeyboardFocus( WidgetToFocusPath, EFocusCause::SetDirectly );
 		bPendingFocusNextFrame = false;
 	}
 
diff --git a/Engine/Source/Editor/DetailCustomizations/Private/ComponentTransformDetails.cpp b/Engine/Source/Editor/DetailCustomizations/Private/ComponentTransformDetails.cpp
--- a/Engine/Source/Editor/DetailCustomizations/Private/ComponentTransformDetails.cpp
+++ b/Engine/Source/Editor/DetailCustomizations/Private/ComponentTransformDetails.cpp
@@ -865,19 +865,19 @@ void FComponentTransformDetails::ExtendZScaleContextMenu( FMenuBuilder& MenuBuil
 
 void FComponentTransformDetails::OnXScaleMirrored()
 {
-	FSlateApplication::Get().ClearKeyboardFocus(EKeyboardFocusCause::Mouse);
+	FSlateApplication::Get().ClearKeyboardFocus(EFocusCause::Mouse);
 	ScaleObject( 1.0f, 0, true, LOCTEXT( "MirrorActorScaleX", "Mirror actor scale X" ) );
 }
 
 void FComponentTransformDetails::OnYScaleMirrored()
 {
-	FSlateApplication::Get().ClearKeyboardFocus(EKeyboardFocusCause::Mouse);
+	FSlateApplication::Get().ClearKeyboardFocus(EFocusCause::Mouse);
 	ScaleObject(1.0f, 1, true, LOCTEXT("MirrorActorScaleY", "Mirror actor scale Y"));
 }
 
 void FComponentTransformDetails::OnZScaleMirrored()
 {
-	FSlateApplication::Get().ClearKeyboardFocus(EKeyboardFocusCause::Mouse);
+	FSlateApplication::Get().ClearKeyboardFocus(EFocusCause::Mouse);
 	ScaleObject(1.0f, 2, true, LOCTEXT("MirrorActorScaleZ", "Mirror actor scale Z"));
 }
 
diff --git a/Engine/Source/Editor/DetailCustomizations/Private/DialogueWaveWidgets.cpp b/Engine/Source/Editor/DetailCustomizations/Private/DialogueWaveWidgets.cpp
--- a/Engine/Source/Editor/DetailCustomizations/Private/DialogueWaveWidgets.cpp
+++ b/Engine/Source/Editor/DetailCustomizations/Private/DialogueWaveWidgets.cpp
@@ -581,7 +581,7 @@ FReply SRemovableDialogueVoicePropertyEditor::OnMouseButtonDown( const FGeometry
 		bIsPressed = true;
 
 		//we need to capture the mouse for MouseUp events
-		Reply =  FReply::Handled().CaptureMouse( AsShared() ).SetKeyboardFocus( AsShared(), EKeyboardFocusCause::Mouse );
+		Reply = FReply::Handled().CaptureMouse(AsShared()).SetUserFocus(AsShared(), EFocusCause::Mouse);
 	}
 
 	//return the constructed reply
diff --git a/Engine/Source/Editor/EditorStyle/Private/SlateEditorStyle.cpp b/Engine/Source/Editor/EditorStyle/Private/SlateEditorStyle.cpp
--- a/Engine/Source/Editor/EditorStyle/Private/SlateEditorStyle.cpp
+++ b/Engine/Source/Editor/EditorStyle/Private/SlateEditorStyle.cpp
@@ -3677,7 +3677,7 @@ void FSlateEditorStyle::FStyle::SetupGraphEditorStyles()
 		Set( "GraphEditor.SpawnActor_16x", new IMAGE_BRUSH("Icons/icon_Blueprint_SpawnActor_16x", Icon16x16));
 		Set( "GraphEditor.PadEvent_16x", new IMAGE_BRUSH("Icons/icon_Blueprint_PadEvent_16x", Icon16x16));
 		Set( "GraphEditor.MouseEvent_16x", new IMAGE_BRUSH("Icons/icon_Blueprint_MouseEvent_16x", Icon16x16));
-		Set( "GraphEditor.KeyEvent_16x", new IMAGE_BRUSH("Icons/icon_Blueprint_KeyboardEvent_16x", Icon16x16));
+		Set( "GraphEditor.KeyEvent_16x", new IMAGE_BRUSH("Icons/icon_Blueprint_KeyEvent_16x", Icon16x16));
 		Set( "GraphEditor.TouchEvent_16x", new IMAGE_BRUSH("Icons/icon_Blueprint_TouchEvent_16x", Icon16x16));
 		Set( "GraphEditor.MakeArray_16x", new IMAGE_BRUSH("Icons/icon_Blueprint_MakeArray_16x", Icon16x16));
 		Set( "GraphEditor.Enum_16x", new IMAGE_BRUSH("Icons/icon_Blueprint_Enum_16x", Icon16x16));
diff --git a/Engine/Source/Editor/EditorWidgets/Private/SAssetSearchBox.cpp b/Engine/Source/Editor/EditorWidgets/Private/SAssetSearchBox.cpp
--- a/Engine/Source/Editor/EditorWidgets/Private/SAssetSearchBox.cpp
+++ b/Engine/Source/Editor/EditorWidgets/Private/SAssetSearchBox.cpp
@@ -53,9 +53,9 @@ void SAssetSearchBox::SetText(const TAttribute< FText >& InNewText)
 	PreCommittedText = InNewText.Get();
 }
 
-FReply SAssetSearchBox::OnPreviewKeyDown( const FGeometry& MyGeometry, const FKeyboardEvent& InKeyboardEvent )
+FReply SAssetSearchBox::OnPreviewKeyDown( const FGeometry& MyGeometry, const FKeyEvent& InKeyEvent )
 {
-	if ( SuggestionBox->IsOpen() && InKeyboardEvent.GetKey() == EKeys::Escape )
+	if ( SuggestionBox->IsOpen() && InKeyEvent.GetKey() == EKeys::Escape )
 	{
 		// Clear any selection first to prevent the currently selection being set in the text box
 		SuggestionListView->ClearSelection();
@@ -66,11 +66,11 @@ FReply SAssetSearchBox::OnPreviewKeyDown( const FGeometry& MyGeometry, const FKe
 	return FReply::Unhandled();
 }
 
-FReply SAssetSearchBox::OnKeyDown( const FGeometry& MyGeometry, const FKeyboardEvent& InKeyboardEvent )
+FReply SAssetSearchBox::OnKeyDown( const FGeometry& MyGeometry, const FKeyEvent& InKeyEvent )
 {
-	if ( SuggestionBox->IsOpen() && (InKeyboardEvent.GetKey() == EKeys::Up || InKeyboardEvent.GetKey() == EKeys::Down) )
+	if ( SuggestionBox->IsOpen() && (InKeyEvent.GetKey() == EKeys::Up || InKeyEvent.GetKey() == EKeys::Down) )
 	{
-		const bool bSelectingUp = InKeyboardEvent.GetKey() == EKeys::Up;
+		const bool bSelectingUp = InKeyEvent.GetKey() == EKeys::Up;
 		TSharedPtr<FString> SelectedSuggestion = GetSelectedSuggestion();
 
 		if ( SelectedSuggestion.IsValid() )
@@ -122,10 +122,10 @@ bool SAssetSearchBox::HasKeyboardFocus() const
 	return InputText->HasKeyboardFocus();
 }
 
-FReply SAssetSearchBox::OnKeyboardFocusReceived( const FGeometry& MyGeometry, const FKeyboardFocusEvent& InKeyboardFocusEvent )
+FReply SAssetSearchBox::OnFocusReceived( const FGeometry& MyGeometry, const FFocusEvent& InFocusEvent )
 {
 	// Forward keyboard focus to our editable text widget
-	return FReply::Handled().SetKeyboardFocus( InputText.ToSharedRef(), InKeyboardFocusEvent.GetCause() );
+	return FReply::Handled().SetUserFocus(InputText.ToSharedRef(), InFocusEvent.GetCause());
 }
 
 void SAssetSearchBox::HandleTextChanged(const FText& NewText)
@@ -250,7 +250,7 @@ void SAssetSearchBox::FocusEditBox()
 {
 	FWidgetPath WidgetToFocusPath;
 	FSlateApplication::Get().GeneratePathToWidgetUnchecked( InputText.ToSharedRef(), WidgetToFocusPath );
-	FSlateApplication::Get().SetKeyboardFocus( WidgetToFocusPath, EKeyboardFocusCause::SetDirectly );
+	FSlateApplication::Get().SetKeyboardFocus( WidgetToFocusPath, EFocusCause::SetDirectly );
 }
 
 TSharedPtr<FString> SAssetSearchBox::GetSelectedSuggestion()
diff --git a/Engine/Source/Editor/EditorWidgets/Public/SAssetSearchBox.h b/Engine/Source/Editor/EditorWidgets/Public/SAssetSearchBox.h
--- a/Engine/Source/Editor/EditorWidgets/Public/SAssetSearchBox.h
+++ b/Engine/Source/Editor/EditorWidgets/Public/SAssetSearchBox.h
@@ -53,11 +53,11 @@ public:
 	void SetText(const TAttribute< FText >& InNewText);
 
 	// SWidget implementation
-	virtual FReply OnPreviewKeyDown( const FGeometry& MyGeometry, const FKeyboardEvent& InKeyboardEvent ) override;
-	virtual FReply OnKeyDown( const FGeometry& MyGeometry, const FKeyboardEvent& InKeyboardEvent ) override;
+	virtual FReply OnPreviewKeyDown( const FGeometry& MyGeometry, const FKeyEvent& InKeyEvent ) override;
+	virtual FReply OnKeyDown( const FGeometry& MyGeometry, const FKeyEvent& InKeyEvent ) override;
 	virtual bool SupportsKeyboardFocus() const override;
 	virtual bool HasKeyboardFocus() const override;
-	virtual FReply OnKeyboardFocusReceived( const FGeometry& MyGeometry, const FKeyboardFocusEvent& InKeyboardFocusEvent ) override;
+	virtual FReply OnFocusReceived( const FGeometry& MyGeometry, const FFocusEvent& InFocusEvent ) override;
 
 private:
 	/** Handler for when text in the editable text box changed */
diff --git a/Engine/Source/Editor/GameProjectGeneration/Private/SNewClassDialog.cpp b/Engine/Source/Editor/GameProjectGeneration/Private/SNewClassDialog.cpp
--- a/Engine/Source/Editor/GameProjectGeneration/Private/SNewClassDialog.cpp
+++ b/Engine/Source/Editor/GameProjectGeneration/Private/SNewClassDialog.cpp
@@ -894,7 +894,7 @@ void SNewClassDialog::OnNamePageEntered()
 	UpdateInputValidity();
 
 	// Steal keyboard focus to accelerate name entering
-	FSlateApplication::Get().SetKeyboardFocus(ClassNameEditBox, EKeyboardFocusCause::SetDirectly);
+	FSlateApplication::Get().SetKeyboardFocus(ClassNameEditBox, EFocusCause::SetDirectly);
 }
 
 FString SNewClassDialog::GetNameClassTitle() const
diff --git a/Engine/Source/Editor/GameProjectGeneration/Private/SProjectBrowser.cpp b/Engine/Source/Editor/GameProjectGeneration/Private/SProjectBrowser.cpp
--- a/Engine/Source/Editor/GameProjectGeneration/Private/SProjectBrowser.cpp
+++ b/Engine/Source/Editor/GameProjectGeneration/Private/SProjectBrowser.cpp
@@ -944,9 +944,9 @@ void SProjectBrowser::PopulateFilteredProjectCategories()
 	}
 }
 
-FReply SProjectBrowser::OnKeyDown( const FGeometry& MyGeometry, const FKeyboardEvent& InKeyboardEvent )
+FReply SProjectBrowser::OnKeyDown( const FGeometry& MyGeometry, const FKeyEvent& InKeyEvent )
 {
-	if (InKeyboardEvent.GetKey() == EKeys::F5)
+	if (InKeyEvent.GetKey() == EKeys::F5)
 	{
 		return FindProjects();
 	}
diff --git a/Engine/Source/Editor/GameProjectGeneration/Private/SProjectBrowser.h b/Engine/Source/Editor/GameProjectGeneration/Private/SProjectBrowser.h
--- a/Engine/Source/Editor/GameProjectGeneration/Private/SProjectBrowser.h
+++ b/Engine/Source/Editor/GameProjectGeneration/Private/SProjectBrowser.h
@@ -80,14 +80,14 @@ protected:
 	void PopulateFilteredProjectCategories();
 	
 	/**
-	 * Called after a key is pressed when this widget has keyboard focus (this event bubbles if not handled)
+	 * Called after a key is pressed when this widget has focus (this event bubbles if not handled)
 	 *
 	 * @param MyGeometry The Geometry of the widget receiving the event
-	 * @param  InKeyboardEvent  Keyboard event
+	 * @param  InKeyEvent  Key event
 	 *
 	 * @return  Returns whether the event was handled, along with other possible actions
 	 */
-	virtual FReply OnKeyDown( const FGeometry& MyGeometry, const FKeyboardEvent& InKeyboardEvent ) override;
+	virtual FReply OnKeyDown( const FGeometry& MyGeometry, const FKeyEvent& InKeyEvent ) override;
 
 private:
 
diff --git a/Engine/Source/Editor/GameProjectGeneration/Private/SVerbChoiceDialog.cpp b/Engine/Source/Editor/GameProjectGeneration/Private/SVerbChoiceDialog.cpp
--- a/Engine/Source/Editor/GameProjectGeneration/Private/SVerbChoiceDialog.cpp
+++ b/Engine/Source/Editor/GameProjectGeneration/Private/SVerbChoiceDialog.cpp
@@ -111,15 +111,15 @@ void SVerbChoiceDialog::Construct( const FArguments& InArgs )
 }
 END_SLATE_FUNCTION_BUILD_OPTIMIZATION
 
-FReply SVerbChoiceDialog::OnKeyDown( const FGeometry& MyGeometry, const FKeyboardEvent& InKeyboardEvent )
+FReply SVerbChoiceDialog::OnKeyDown( const FGeometry& MyGeometry, const FKeyEvent& InKeyEvent )
 {
 	//see if we pressed the Enter or Spacebar keys
-	if( InKeyboardEvent.GetKey() == EKeys::Escape )
+	if( InKeyEvent.GetKey() == EKeys::Escape )
 	{
 		return HandleButtonClicked(EAppReturnType::Cancel);
 	}
 
-	if (InKeyboardEvent.GetKey() == EKeys::C && InKeyboardEvent.IsControlDown())
+	if (InKeyEvent.GetKey() == EKeys::C && InKeyEvent.IsControlDown())
 	{
 		CopyMessageToClipboard();
 
diff --git a/Engine/Source/Editor/GameProjectGeneration/Private/SVerbChoiceDialog.h b/Engine/Source/Editor/GameProjectGeneration/Private/SVerbChoiceDialog.h
--- a/Engine/Source/Editor/GameProjectGeneration/Private/SVerbChoiceDialog.h
+++ b/Engine/Source/Editor/GameProjectGeneration/Private/SVerbChoiceDialog.h
@@ -19,7 +19,7 @@ public:
 	void Construct( const FArguments& InArgs );
 
 	/** Override behavior when a key is pressed */
-	virtual	FReply OnKeyDown( const FGeometry& MyGeometry, const FKeyboardEvent& InKeyboardEvent ) override;
+	virtual	FReply OnKeyDown( const FGeometry& MyGeometry, const FKeyEvent& InKeyEvent ) override;
 
 	/** Override the base method to allow for keyboard focus */
 	virtual bool SupportsKeyboardFocus() const override;
diff --git a/Engine/Source/Editor/GameplayAbilitiesEditor/Private/GameplayAbilitiesBlueprintFactory.cpp b/Engine/Source/Editor/GameplayAbilitiesEditor/Private/GameplayAbilitiesBlueprintFactory.cpp
--- a/Engine/Source/Editor/GameplayAbilitiesEditor/Private/GameplayAbilitiesBlueprintFactory.cpp
+++ b/Engine/Source/Editor/GameplayAbilitiesEditor/Private/GameplayAbilitiesBlueprintFactory.cpp
@@ -203,14 +203,14 @@ private:
 		return FReply::Handled();
 	}
 
-	FReply OnKeyDown(const FGeometry& MyGeometry, const FKeyboardEvent& InKeyboardEvent)
+	FReply OnKeyDown(const FGeometry& MyGeometry, const FKeyEvent& InKeyEvent)
 	{
-		if (InKeyboardEvent.GetKey() == EKeys::Escape)
+		if (InKeyEvent.GetKey() == EKeys::Escape)
 		{
 			CloseDialog();
 			return FReply::Handled();
 		}
-		return SWidget::OnKeyDown(MyGeometry, InKeyboardEvent);
+		return SWidget::OnKeyDown(MyGeometry, InKeyEvent);
 	}
 
 private:
diff --git a/Engine/Source/Editor/GraphEditor/Private/DragConnection.cpp b/Engine/Source/Editor/GraphEditor/Private/DragConnection.cpp
--- a/Engine/Source/Editor/GraphEditor/Private/DragConnection.cpp
+++ b/Engine/Source/Editor/GraphEditor/Private/DragConnection.cpp
@@ -205,7 +205,7 @@ FReply FDragConnection::DroppedOnPanel( const TSharedRef< SWidget >& Panel, FVec
 
 	// Give the context menu focus
 	return (WidgetToFocus.IsValid())
-		? FReply::Handled().SetKeyboardFocus( WidgetToFocus.ToSharedRef(), EKeyboardFocusCause::SetDirectly )
+		? FReply::Handled().SetUserFocus(WidgetToFocus.ToSharedRef(), EFocusCause::SetDirectly)
 		: FReply::Handled();
 }
 
diff --git a/Engine/Source/Editor/GraphEditor/Private/SGraphActionMenu.cpp b/Engine/Source/Editor/GraphEditor/Private/SGraphActionMenu.cpp
--- a/Engine/Source/Editor/GraphEditor/Private/SGraphActionMenu.cpp
+++ b/Engine/Source/Editor/GraphEditor/Private/SGraphActionMenu.cpp
@@ -1063,28 +1063,28 @@ bool SGraphActionMenu::OnMouseButtonDownEvent( TWeakPtr<FEdGraphSchemaAction> In
 	return bResult;
 }
 
-FReply SGraphActionMenu::OnKeyDown( const FGeometry& MyGeometry, const FKeyboardEvent& KeyboardEvent )
+FReply SGraphActionMenu::OnKeyDown( const FGeometry& MyGeometry, const FKeyEvent& KeyEvent )
 {
 	int32 SelectionDelta = 0;
 
 	// Escape dismisses the menu without placing a node
-	if (KeyboardEvent.GetKey() == EKeys::Escape)
+	if (KeyEvent.GetKey() == EKeys::Escape)
 	{
 		FSlateApplication::Get().DismissAllMenus();
 		return FReply::Handled();
 	}
-	else if ((KeyboardEvent.GetKey() == EKeys::Enter) && !bIgnoreUIUpdate)
+	else if ((KeyEvent.GetKey() == EKeys::Enter) && !bIgnoreUIUpdate)
 	{
 		return TryToSpawnActiveSuggestion() ? FReply::Handled() : FReply::Unhandled();
 	}
 	else if (FilteredActionNodes.Num() > 0)
 	{
 		// Up and down move thru the filtered node list
-		if (KeyboardEvent.GetKey() == EKeys::Up)
+		if (KeyEvent.GetKey() == EKeys::Up)
 		{
 			SelectionDelta = -1;
 		}
-		else if (KeyboardEvent.GetKey() == EKeys::Down)
+		else if (KeyEvent.GetKey() == EKeys::Down)
 		{
 			SelectionDelta = +1;
 		}
@@ -1098,7 +1098,7 @@ FReply SGraphActionMenu::OnKeyDown( const FGeometry& MyGeometry, const FKeyboard
 				TGuardValue<bool> PreventSelectionFromTriggeringCommit(bIgnoreUIUpdate, true);
 				TreeView->SetSelection(FilteredRootAction->Children[SelectedSuggestion], ESelectInfo::OnKeyPress);
 				TreeView->RequestScrollIntoView(FilteredRootAction->Children[SelectedSuggestion]);
-				return FReply::Handled().SetKeyboardFocus( SharedThis( TreeView.Get() ), EKeyboardFocusCause::WindowActivate );
+				return FReply::Handled().SetUserFocus(SharedThis(TreeView.Get()), EFocusCause::WindowActivate);
 			}
 
 			//Move up or down one, wrapping around
diff --git a/Engine/Source/Editor/GraphEditor/Private/SGraphEditorImpl.cpp b/Engine/Source/Editor/GraphEditor/Private/SGraphEditorImpl.cpp
--- a/Engine/Source/Editor/GraphEditor/Private/SGraphEditorImpl.cpp
+++ b/Engine/Source/Editor/GraphEditor/Private/SGraphEditorImpl.cpp
@@ -41,7 +41,7 @@ bool SGraphEditorImpl::SupportsKeyboardFocus() const
 	return true;
 }
 
-FReply SGraphEditorImpl::OnKeyboardFocusReceived( const FGeometry& MyGeometry, const FKeyboardFocusEvent& InKeyboardFocusEvent )
+FReply SGraphEditorImpl::OnFocusReceived( const FGeometry& MyGeometry, const FFocusEvent& InFocusEvent )
 {
 	OnFocused.ExecuteIfBound(SharedThis(this));
 	return FReply::Handled();
@@ -57,15 +57,15 @@ FReply SGraphEditorImpl::OnMouseButtonDown( const FGeometry& MyGeometry, const F
 	{
 		OnNavigateHistoryForward.ExecuteIfBound();
 	}
-	return FReply::Handled().SetKeyboardFocus( SharedThis(this), EKeyboardFocusCause::Mouse );
+	return FReply::Handled().SetUserFocus(SharedThis(this), EFocusCause::Mouse);
 }
 
-FReply SGraphEditorImpl::OnKeyDown( const FGeometry& MyGeometry, const FKeyboardEvent& InKeyboardEvent )
+FReply SGraphEditorImpl::OnKeyDown( const FGeometry& MyGeometry, const FKeyEvent& InKeyEvent )
 {
 	int32 NumNodes = GetCurrentGraph()->Nodes.Num();
-	if (Commands->ProcessCommandBindings( InKeyboardEvent ) )
+	if (Commands->ProcessCommandBindings( InKeyEvent ) )
 	{
-		bool bPasteOperation = InKeyboardEvent.IsControlDown() && InKeyboardEvent.GetKey() == EKeys::V;
+		bool bPasteOperation = InKeyEvent.IsControlDown() && InKeyEvent.GetKey() == EKeys::V;
 
 		if(	!bPasteOperation && GetCurrentGraph()->Nodes.Num() > NumNodes )
 		{
@@ -75,7 +75,7 @@ FReply SGraphEditorImpl::OnKeyDown( const FGeometry& MyGeometry, const FKeyboard
 	}
 	else
 	{
-		return SCompoundWidget::OnKeyDown(MyGeometry, InKeyboardEvent);
+		return SCompoundWidget::OnKeyDown(MyGeometry, InKeyEvent);
 	}
 }
 
diff --git a/Engine/Source/Editor/GraphEditor/Private/SGraphEditorImpl.h b/Engine/Source/Editor/GraphEditor/Private/SGraphEditorImpl.h
--- a/Engine/Source/Editor/GraphEditor/Private/SGraphEditorImpl.h
+++ b/Engine/Source/Editor/GraphEditor/Private/SGraphEditorImpl.h
@@ -115,9 +115,9 @@ public:
 
 	void FocusLockedEditorHere();
 
-	virtual FReply OnKeyboardFocusReceived( const FGeometry& MyGeometry, const FKeyboardFocusEvent& InKeyboardFocusEvent ) override;
+	virtual FReply OnFocusReceived( const FGeometry& MyGeometry, const FFocusEvent& InFocusEvent ) override;
 	virtual FReply OnMouseButtonDown( const FGeometry& MyGeometry, const FPointerEvent& MouseEvent ) override;
-	virtual FReply OnKeyDown( const FGeometry& MyGeometry, const FKeyboardEvent& InKeyboardEvent ) override;
+	virtual FReply OnKeyDown( const FGeometry& MyGeometry, const FKeyEvent& InKeyEvent ) override;
 	virtual bool SupportsKeyboardFocus() const override;
 	// End of SWidget interface
 
diff --git a/Engine/Source/Editor/GraphEditor/Private/SGraphNodeComment.cpp b/Engine/Source/Editor/GraphEditor/Private/SGraphNodeComment.cpp
--- a/Engine/Source/Editor/GraphEditor/Private/SGraphNodeComment.cpp
+++ b/Engine/Source/Editor/GraphEditor/Private/SGraphNodeComment.cpp
@@ -197,7 +197,7 @@ FReply SGraphNodeComment::OnMouseButtonDoubleClick( const FGeometry& InMyGeometr
 		// Set the keyboard focus
 		if(!HasKeyboardFocus())
 		{
-			FSlateApplication::Get().SetKeyboardFocus(SharedThis(this), EKeyboardFocusCause::SetDirectly);
+			FSlateApplication::Get().SetKeyboardFocus(SharedThis(this), EFocusCause::SetDirectly);
 		}
 
 		return FReply::Handled();
diff --git a/Engine/Source/Editor/GraphEditor/Private/SGraphPanel.cpp b/Engine/Source/Editor/GraphEditor/Private/SGraphPanel.cpp
--- a/Engine/Source/Editor/GraphEditor/Private/SGraphPanel.cpp
+++ b/Engine/Source/Editor/GraphEditor/Private/SGraphPanel.cpp
@@ -541,44 +541,44 @@ void SGraphPanel::UpdateSelectedNodesPositions (FVector2D PositionIncrement)
 	}
 }
 
-FReply SGraphPanel::OnKeyDown( const FGeometry& MyGeometry, const FKeyboardEvent& InKeyboardEvent )
+FReply SGraphPanel::OnKeyDown( const FGeometry& MyGeometry, const FKeyEvent& InKeyEvent )
 {
 	if( IsEditable.Get() )
 	{
-		if( InKeyboardEvent.GetKey() == EKeys::Up  ||InKeyboardEvent.GetKey() ==  EKeys::NumPadEight )
+		if( InKeyEvent.GetKey() == EKeys::Up  ||InKeyEvent.GetKey() ==  EKeys::NumPadEight )
 		{
 			UpdateSelectedNodesPositions(FVector2D(0.0f,-GetSnapGridSize()));
 			return FReply::Handled();
 		}
-		if( InKeyboardEvent.GetKey() ==  EKeys::Down || InKeyboardEvent.GetKey() ==  EKeys::NumPadTwo )
+		if( InKeyEvent.GetKey() ==  EKeys::Down || InKeyEvent.GetKey() ==  EKeys::NumPadTwo )
 		{
 			UpdateSelectedNodesPositions(FVector2D(0.0f,GetSnapGridSize()));
 			return FReply::Handled();
 		}
-		if( InKeyboardEvent.GetKey() ==  EKeys::Right || InKeyboardEvent.GetKey() ==  EKeys::NumPadSix )
+		if( InKeyEvent.GetKey() ==  EKeys::Right || InKeyEvent.GetKey() ==  EKeys::NumPadSix )
 		{
 			UpdateSelectedNodesPositions(FVector2D(GetSnapGridSize(),0.0f));
 			return FReply::Handled();
 		}
-		if( InKeyboardEvent.GetKey() ==  EKeys::Left || InKeyboardEvent.GetKey() ==  EKeys::NumPadFour )
+		if( InKeyEvent.GetKey() ==  EKeys::Left || InKeyEvent.GetKey() ==  EKeys::NumPadFour )
 		{
 			UpdateSelectedNodesPositions(FVector2D(-GetSnapGridSize(),0.0f));
 			return FReply::Handled();
 		}
-		if ( InKeyboardEvent.GetKey() ==  EKeys::Subtract )
+		if ( InKeyEvent.GetKey() ==  EKeys::Subtract )
 		{
-			ChangeZoomLevel(-1, CachedAllottedGeometryScaledSize / 2.f, InKeyboardEvent.IsControlDown());
+			ChangeZoomLevel(-1, CachedAllottedGeometryScaledSize / 2.f, InKeyEvent.IsControlDown());
 			return FReply::Handled();
 		}
-		if ( InKeyboardEvent.GetKey() ==  EKeys::Add )
+		if ( InKeyEvent.GetKey() ==  EKeys::Add )
 		{
-			ChangeZoomLevel(+1, CachedAllottedGeometryScaledSize / 2.f, InKeyboardEvent.IsControlDown());
+			ChangeZoomLevel(+1, CachedAllottedGeometryScaledSize / 2.f, InKeyEvent.IsControlDown());
 			return FReply::Handled();
 		}
 
 	}
 
-	return SNodePanel::OnKeyDown(MyGeometry, InKeyboardEvent);
+	return SNodePanel::OnKeyDown(MyGeometry, InKeyEvent);
 }
 
 
@@ -810,7 +810,7 @@ FReply SGraphPanel::OnDrop( const FGeometry& MyGeometry, const FDragDropEvent& D
 {
 	const FVector2D NodeAddPosition = PanelCoordToGraphCoord( MyGeometry.AbsoluteToLocal( DragDropEvent.GetScreenSpacePosition() ) );
 
-	FSlateApplication::Get().SetKeyboardFocus(AsShared(), EKeyboardFocusCause::SetDirectly);
+	FSlateApplication::Get().SetKeyboardFocus(AsShared(), EFocusCause::SetDirectly);
 
 	TSharedPtr<FDragDropOperation> Operation = DragDropEvent.GetOperation();
 	if (!Operation.IsValid() || !IsEditable.Get())
diff --git a/Engine/Source/Editor/GraphEditor/Private/SNodePanel.cpp b/Engine/Source/Editor/GraphEditor/Private/SNodePanel.cpp
--- a/Engine/Source/Editor/GraphEditor/Private/SNodePanel.cpp
+++ b/Engine/Source/Editor/GraphEditor/Private/SNodePanel.cpp
@@ -810,7 +810,7 @@ FReply SNodePanel::OnMouseButtonUp( const FGeometry& MyGeometry, const FPointerE
 
 		if (WidgetToFocus.IsValid())
 		{
-			ReplyState.SetKeyboardFocus( WidgetToFocus.ToSharedRef(), EKeyboardFocusCause::SetDirectly );
+			ReplyState.SetUserFocus(WidgetToFocus.ToSharedRef(), EFocusCause::SetDirectly);
 		}
 	}
 	else if ( bIsLeftMouseButtonEffecting )
@@ -914,23 +914,23 @@ FCursorReply SNodePanel::OnCursorQuery( const FGeometry& MyGeometry, const FPoin
 		FCursorReply::Cursor( EMouseCursor::Default );
 }
 
-FReply SNodePanel::OnKeyDown( const FGeometry& MyGeometry, const FKeyboardEvent& InKeyboardEvent )
+FReply SNodePanel::OnKeyDown( const FGeometry& MyGeometry, const FKeyEvent& InKeyEvent )
 {
 	if( IsEditable.Get() )
 	{
-		LastKeyGestureDetected.Key = InKeyboardEvent.GetKey();
-		LastKeyGestureDetected.bAlt = InKeyboardEvent.IsAltDown();
-		LastKeyGestureDetected.bCtrl = InKeyboardEvent.IsControlDown();
-		LastKeyGestureDetected.bShift = InKeyboardEvent.IsShiftDown();
-		LastKeyGestureDetected.bCmd = InKeyboardEvent.IsCommandDown();
+		LastKeyGestureDetected.Key = InKeyEvent.GetKey();
+		LastKeyGestureDetected.bAlt = InKeyEvent.IsAltDown();
+		LastKeyGestureDetected.bCtrl = InKeyEvent.IsControlDown();
+		LastKeyGestureDetected.bShift = InKeyEvent.IsShiftDown();
+		LastKeyGestureDetected.bCmd = InKeyEvent.IsCommandDown();
 	}
 
 	return FReply::Unhandled();
 }
 
-FReply SNodePanel::OnKeyUp( const FGeometry& MyGeometry, const FKeyboardEvent& InKeyboardEvent )
+FReply SNodePanel::OnKeyUp( const FGeometry& MyGeometry, const FKeyEvent& InKeyEvent )
 {
-	if(LastKeyGestureDetected.Key == InKeyboardEvent.GetKey())
+	if(LastKeyGestureDetected.Key == InKeyEvent.GetKey())
 	{
 		LastKeyGestureDetected.Key = EKeys::Invalid;
 		LastKeyGestureDetected.bAlt = false;
@@ -942,7 +942,7 @@ FReply SNodePanel::OnKeyUp( const FGeometry& MyGeometry, const FKeyboardEvent& I
 	return FReply::Unhandled();
 }
 
-void SNodePanel::OnKeyboardFocusLost( const FKeyboardFocusEvent& InKeyboardFocusEvent )
+void SNodePanel::OnFocusLost( const FFocusEvent& InFocusEvent )
 {
 	LastKeyGestureDetected.Key = EKeys::Invalid;
 	LastKeyGestureDetected.bAlt = false;
diff --git a/Engine/Source/Editor/GraphEditor/Public/SGraphActionMenu.h b/Engine/Source/Editor/GraphEditor/Public/SGraphActionMenu.h
--- a/Engine/Source/Editor/GraphEditor/Public/SGraphActionMenu.h
+++ b/Engine/Source/Editor/GraphEditor/Public/SGraphActionMenu.h
@@ -181,7 +181,7 @@ protected:
 
 public:
 	// SWidget interface
-	virtual FReply OnKeyDown(const FGeometry& MyGeometry, const FKeyboardEvent& KeyboardEvent) override;
+	virtual FReply OnKeyDown(const FGeometry& MyGeometry, const FKeyEvent& KeyEvent) override;
 	// End of SWidget interface
 
 	/** Get filter text box widget */
diff --git a/Engine/Source/Editor/GraphEditor/Public/SNodePanel.h b/Engine/Source/Editor/GraphEditor/Public/SNodePanel.h
--- a/Engine/Source/Editor/GraphEditor/Public/SNodePanel.h
+++ b/Engine/Source/Editor/GraphEditor/Public/SNodePanel.h
@@ -324,9 +324,9 @@ public:
 	virtual FReply OnMouseMove(const FGeometry& MyGeometry, const FPointerEvent& MouseEvent) override;
 	virtual FReply OnMouseWheel(const FGeometry& MyGeometry, const FPointerEvent& MouseEvent) override;
 	virtual FCursorReply OnCursorQuery( const FGeometry& MyGeometry, const FPointerEvent& CursorEvent ) const override;
-	virtual FReply OnKeyDown( const FGeometry& MyGeometry, const FKeyboardEvent& InKeyboardEvent ) override;
-	virtual FReply OnKeyUp( const FGeometry& MyGeometry, const FKeyboardEvent& InKeyboardEvent ) override;
-	virtual void OnKeyboardFocusLost( const FKeyboardFocusEvent& InKeyboardFocusEvent ) override;
+	virtual FReply OnKeyDown( const FGeometry& MyGeometry, const FKeyEvent& InKeyEvent ) override;
+	virtual FReply OnKeyUp( const FGeometry& MyGeometry, const FKeyEvent& InKeyEvent ) override;
+	virtual void OnFocusLost( const FFocusEvent& InFocusEvent ) override;
 	virtual FReply OnTouchGesture( const FGeometry& MyGeometry, const FPointerEvent& GestureEvent ) override;
 	virtual FReply OnTouchEnded( const FGeometry& MyGeometry, const FPointerEvent& InTouchEvent ) override;
 
diff --git a/Engine/Source/Editor/InputBindingEditor/Private/Widgets/SGestureEditBox.cpp b/Engine/Source/Editor/InputBindingEditor/Private/Widgets/SGestureEditBox.cpp
--- a/Engine/Source/Editor/InputBindingEditor/Private/Widgets/SGestureEditBox.cpp
+++ b/Engine/Source/Editor/InputBindingEditor/Private/Widgets/SGestureEditBox.cpp
@@ -144,7 +144,7 @@ void SGestureEditBox::OnGestureChanged()
 		{
 			GestureEditor->CommitNewGesture();
 			GestureEditor->StopEditing();
-			FSlateApplication::Get().ClearKeyboardFocus(EKeyboardFocusCause::SetDirectly);
+			FSlateApplication::Get().ClearKeyboardFocus(EFocusCause::SetDirectly);
 		}
 	}
 }
@@ -260,7 +260,7 @@ FReply SGestureEditBox::OnMouseButtonDown( const FGeometry& MyGeometry, const FP
 		{
 			GestureEditor->StartEditing();
 		}
-		return FReply::Handled().SetKeyboardFocus( GestureEditor.ToSharedRef(), EKeyboardFocusCause::Mouse );
+		return FReply::Handled().SetUserFocus(GestureEditor.ToSharedRef(), EFocusCause::Mouse);
 	}
 
 	return FReply::Handled();
diff --git a/Engine/Source/Editor/InputBindingEditor/Private/Widgets/SGestureEditor.cpp b/Engine/Source/Editor/InputBindingEditor/Private/Widgets/SGestureEditor.cpp
--- a/Engine/Source/Editor/InputBindingEditor/Private/Widgets/SGestureEditor.cpp
+++ b/Engine/Source/Editor/InputBindingEditor/Private/Widgets/SGestureEditor.cpp
@@ -38,9 +38,9 @@ void SGestureEditor::Construct( const FArguments& InArgs, TSharedPtr<FGestureTre
 TWeakPtr<SGestureEditor> SGestureEditor::GestureBeingEdited;
 
 
-FReply SGestureEditor::OnKeyDown( const FGeometry& MyGeometry, const FKeyboardEvent& InKeyboardEvent )
+FReply SGestureEditor::OnKeyDown( const FGeometry& MyGeometry, const FKeyEvent& InKeyEvent )
 {
-	const FKey Key = InKeyboardEvent.GetKey();
+	const FKey Key = InKeyEvent.GetKey();
 
 	if( bIsEditing ) 
 	{
@@ -54,10 +54,10 @@ FReply SGestureEditor::OnKeyDown( const FGeometry& MyGeometry, const FKeyboardEv
 	
 		StartChangingText();
 
-		EditingInputGesture.bCtrl = InKeyboardEvent.IsControlDown();
-		EditingInputGesture.bAlt = InKeyboardEvent.IsAltDown();
-		EditingInputGesture.bShift = InKeyboardEvent.IsShiftDown();
-		EditingInputGesture.bCmd = InKeyboardEvent.IsCommandDown();
+		EditingInputGesture.bCtrl = InKeyEvent.IsControlDown();
+		EditingInputGesture.bAlt = InKeyEvent.IsAltDown();
+		EditingInputGesture.bShift = InKeyEvent.IsShiftDown();
+		EditingInputGesture.bCmd = InKeyEvent.IsCommandDown();
 
 		LoadText();
 		//@todo checking the length of localized string is not valid, at the very least in this manner [10/11/2013 justin.sargent]
@@ -74,7 +74,7 @@ FReply SGestureEditor::OnKeyDown( const FGeometry& MyGeometry, const FKeyboardEv
 }
 
 
-FReply SGestureEditor::OnKeyUp( const FGeometry& MyGeometry, const FKeyboardEvent& InKeyboardEvent )
+FReply SGestureEditor::OnKeyUp( const FGeometry& MyGeometry, const FKeyEvent& InKeyEvent )
 {
 	return FReply::Unhandled();
 }
@@ -91,7 +91,7 @@ FReply SGestureEditor::OnMouseButtonDown( const FGeometry& MyGeometry, const FPo
 	if ( InMouseEvent.GetEffectingButton() == EKeys::LeftMouseButton && !bIsEditing )
 	{
 		StartEditing();
-		return FReply::Handled().SetKeyboardFocus( AsShared(), EKeyboardFocusCause::Mouse );
+		return FReply::Handled().SetUserFocus(AsShared(), EFocusCause::Mouse);
 	}
 
 	return FReply::Unhandled();
@@ -104,7 +104,7 @@ FReply SGestureEditor::OnMouseButtonDoubleClick( const FGeometry& InMyGeometry,
 }
 
 
-void SGestureEditor::OnKeyboardFocusLost( const FKeyboardFocusEvent& InKeyboardFocusEvent )
+void SGestureEditor::OnFocusLost( const FFocusEvent& InFocusEvent )
 {
 	// Notify a listener that we lost focus so they can determine if we should still be in edit mode
 	OnEditBoxLostFocus.ExecuteIfBound();
diff --git a/Engine/Source/Editor/InputBindingEditor/Private/Widgets/SGestureEditor.h b/Engine/Source/Editor/InputBindingEditor/Private/Widgets/SGestureEditor.h
--- a/Engine/Source/Editor/InputBindingEditor/Private/Widgets/SGestureEditor.h
+++ b/Engine/Source/Editor/InputBindingEditor/Private/Widgets/SGestureEditor.h
@@ -73,11 +73,11 @@ private:
 	/** Gesture that is currently being edited. */
 	static TWeakPtr<SGestureEditor> GestureBeingEdited;
 
-	virtual FReply OnKeyDown( const FGeometry& MyGeometry, const FKeyboardEvent& InKeyboardEvent ) override;
-	virtual FReply OnKeyUp( const FGeometry& MyGeometry, const FKeyboardEvent& InKeyboardEvent ) override;
+	virtual FReply OnKeyDown( const FGeometry& MyGeometry, const FKeyEvent& InKeyEvent ) override;
+	virtual FReply OnKeyUp( const FGeometry& MyGeometry, const FKeyEvent& InKeyEvent ) override;
 	virtual FReply OnMouseButtonDown( const FGeometry& MyGeometry, const FPointerEvent& InMouseEvent ) override;
 	virtual FReply OnMouseButtonDoubleClick( const FGeometry& InMyGeometry, const FPointerEvent& InMouseEvent ) override;
-	virtual void OnKeyboardFocusLost( const FKeyboardFocusEvent& InKeyboardFocusEvent ) override;
+	virtual void OnFocusLost( const FFocusEvent& InFocusEvent ) override;
 	virtual FReply OnKeyChar( const FGeometry& MyGeometry, const FCharacterEvent& InCharacterEvent ) override;
 
 	/** 
diff --git a/Engine/Source/Editor/InputBindingEditor/Private/Widgets/SInputBindingEditorPanel.cpp b/Engine/Source/Editor/InputBindingEditor/Private/Widgets/SInputBindingEditorPanel.cpp
--- a/Engine/Source/Editor/InputBindingEditor/Private/Widgets/SInputBindingEditorPanel.cpp
+++ b/Engine/Source/Editor/InputBindingEditor/Private/Widgets/SInputBindingEditorPanel.cpp
@@ -315,13 +315,13 @@ bool SInputBindingEditorPanel::SupportsKeyboardFocus() const
 }
 
 
-FReply SInputBindingEditorPanel::OnKeyboardFocusReceived( const FGeometry& MyGeometry, const FKeyboardFocusEvent& InKeyboardFocusEvent )
+FReply SInputBindingEditorPanel::OnFocusReceived( const FGeometry& MyGeometry, const FFocusEvent& InFocusEvent )
 {
 	FReply Reply = FReply::Handled();
 
-	if( InKeyboardFocusEvent.GetCause() != EKeyboardFocusCause::Cleared )
+	if( InFocusEvent.GetCause() != EFocusCause::Cleared )
 	{
-		Reply.SetKeyboardFocus( SearchBox.ToSharedRef(), InKeyboardFocusEvent.GetCause() );
+		Reply.SetUserFocus(SearchBox.ToSharedRef(), InFocusEvent.GetCause());
 	}
 
 	return Reply;
diff --git a/Engine/Source/Editor/InputBindingEditor/Private/Widgets/SInputBindingEditorPanel.h b/Engine/Source/Editor/InputBindingEditor/Private/Widgets/SInputBindingEditorPanel.h
--- a/Engine/Source/Editor/InputBindingEditor/Private/Widgets/SInputBindingEditorPanel.h
+++ b/Engine/Source/Editor/InputBindingEditor/Private/Widgets/SInputBindingEditorPanel.h
@@ -76,7 +76,7 @@ public:
 	// SWidget interface
 
 	virtual bool SupportsKeyboardFocus() const override;
-	virtual FReply OnKeyboardFocusReceived( const FGeometry& MyGeometry, const FKeyboardFocusEvent& InKeyboardFocusEvent ) override;
+	virtual FReply OnFocusReceived( const FGeometry& MyGeometry, const FFocusEvent& InFocusEvent ) override;
 
 private:
 
diff --git a/Engine/Source/Editor/IntroTutorials/Private/STutorialEditableText.cpp b/Engine/Source/Editor/IntroTutorials/Private/STutorialEditableText.cpp
--- a/Engine/Source/Editor/IntroTutorials/Private/STutorialEditableText.cpp
+++ b/Engine/Source/Editor/IntroTutorials/Private/STutorialEditableText.cpp
@@ -396,7 +396,7 @@ void STutorialEditableText::StyleSelectedText()
 		const FRunInfo RunInfo = ActiveStyle->CreateRunInfo();
 		const FTextBlockStyle TextBlockStyle = ActiveStyle->CreateTextBlockStyle();
 		RichEditableTextBox->ApplyToSelection(RunInfo, TextBlockStyle);
-		FSlateApplication::Get().SetKeyboardFocus(RichEditableTextBox, EKeyboardFocusCause::SetDirectly);
+		FSlateApplication::Get().SetKeyboardFocus(RichEditableTextBox, EFocusCause::SetDirectly);
 	}
 }
 
diff --git a/Engine/Source/Editor/IntroTutorials/Private/TutorialStructCustomization.cpp b/Engine/Source/Editor/IntroTutorials/Private/TutorialStructCustomization.cpp
--- a/Engine/Source/Editor/IntroTutorials/Private/TutorialStructCustomization.cpp
+++ b/Engine/Source/Editor/IntroTutorials/Private/TutorialStructCustomization.cpp
@@ -267,7 +267,7 @@ private:
 		}
 
 		// kind of a hack, but we need to maintain keyboard focus otherwise we wont get our keypress to 'pick'
-		FSlateApplication::Get().SetKeyboardFocus(SharedThis(this), EKeyboardFocusCause::SetDirectly);
+		FSlateApplication::Get().SetKeyboardFocus(SharedThis(this), EFocusCause::SetDirectly);
 		if(ParentWindow.IsValid())
 		{
 			// also kind of a hack, but this is the only way at the moment to get a 'cursor decorator' without using the drag-drop code path
@@ -280,11 +280,11 @@ private:
 		return FText::Format(LOCTEXT("TootipHint", "{0} (Esc to pick)"), FText::FromName(PickedWidgetName));
 	}
 
-	virtual FReply OnKeyDown(const FGeometry& MyGeometry, const FKeyboardEvent& InKeyboardEvent) override
+	virtual FReply OnKeyDown(const FGeometry& MyGeometry, const FKeyEvent& InKeyEvent) override
 	{
-		if(InKeyboardEvent.GetKey() == EKeys::Escape)
+		if(InKeyEvent.GetKey() == EKeys::Escape)
 		{
-			if( InKeyboardEvent.IsLeftControlDown() == false )
+			if( InKeyEvent.IsLeftControlDown() == false )
 			{
 				// We cant set a parameter if this isn't valid !
 				check(PickPropertyHandle.IsValid());
diff --git a/Engine/Source/Editor/Kismet/Private/FindInBlueprints.cpp b/Engine/Source/Editor/Kismet/Private/FindInBlueprints.cpp
--- a/Engine/Source/Editor/Kismet/Private/FindInBlueprints.cpp
+++ b/Engine/Source/Editor/Kismet/Private/FindInBlueprints.cpp
@@ -1202,7 +1202,7 @@ void SFindInBlueprints::FocusForUse(bool bSetFindWithinBlueprint, FString NewSea
 	FSlateApplication::Get().GeneratePathToWidgetUnchecked( SearchTextField.ToSharedRef(), FilterTextBoxWidgetPath );
 
 	// Set keyboard focus directly
-	FSlateApplication::Get().SetKeyboardFocus( FilterTextBoxWidgetPath, EKeyboardFocusCause::SetDirectly );
+	FSlateApplication::Get().SetKeyboardFocus( FilterTextBoxWidgetPath, EFocusCause::SetDirectly );
 
 	// Set the filter mode
 	bIsInFindWithinBlueprintMode = bSetFindWithinBlueprint;
diff --git a/Engine/Source/Editor/Kismet/Private/SBlueprintDiff.cpp b/Engine/Source/Editor/Kismet/Private/SBlueprintDiff.cpp
--- a/Engine/Source/Editor/Kismet/Private/SBlueprintDiff.cpp
+++ b/Engine/Source/Editor/Kismet/Private/SBlueprintDiff.cpp
@@ -275,7 +275,7 @@ public:
 	void OnSelectionChanged(FSharedDiffOnGraph Item, ESelectInfo::Type SelectionType);
 
 	/** Called when user presses key within the diff view */
-	void KeyWasPressed(const FKeyboardEvent& InKeyboardEvent);
+	void KeyWasPressed(const FKeyEvent& InKeyEvent);
 
 private:
 
@@ -608,11 +608,11 @@ void FListItemGraphToDiff::OnGraphChanged( const FEdGraphEditAction& Action )
 	Diff->OnGraphChanged(this);
 }
 
-void FListItemGraphToDiff::KeyWasPressed( const FKeyboardEvent& InKeyboardEvent )
+void FListItemGraphToDiff::KeyWasPressed( const FKeyEvent& InKeyEvent )
 {
 	if ( KeyCommands.IsValid() )
 	{
-		KeyCommands->ProcessCommandBindings(InKeyboardEvent);
+		KeyCommands->ProcessCommandBindings(InKeyEvent);
 	}
 }
 
@@ -1333,13 +1333,13 @@ void SBlueprintDiff::SetCurrentMode(FName NewMode)
 	}
 }
 
-FReply SBlueprintDiff::OnKeyDown( const FGeometry& MyGeometry, const FKeyboardEvent& InKeyboardEvent ) 
+FReply SBlueprintDiff::OnKeyDown( const FGeometry& MyGeometry, const FKeyEvent& InKeyEvent ) 
 {
 	TArray< FGraphToDiff>  Selected = GraphsToDiff->GetSelectedItems();
 	for(auto It(Selected.CreateIterator());It;++It)
 	{
 		FGraphToDiff& Diff = *It;
-		Diff->KeyWasPressed(InKeyboardEvent);
+		Diff->KeyWasPressed(InKeyEvent);
 	}
 	return FReply::Handled();
 }
diff --git a/Engine/Source/Editor/Kismet/Private/SKismetInspector.cpp b/Engine/Source/Editor/Kismet/Private/SKismetInspector.cpp
--- a/Engine/Source/Editor/Kismet/Private/SKismetInspector.cpp
+++ b/Engine/Source/Editor/Kismet/Private/SKismetInspector.cpp
@@ -300,7 +300,7 @@ void SKismetInspector::ShowDetailsForObjects(const TArray<UObject*>& PropertyObj
 		// if this occurs we need need to immediately clear keyboard focus
 		if ( FSlateApplication::Get().HasFocusedDescendants(AsShared()) )
 		{
-			FSlateApplication::Get().ClearKeyboardFocus(EKeyboardFocusCause::Mouse);
+			FSlateApplication::Get().ClearKeyboardFocus(EFocusCause::Mouse);
 		}
 		bIsReentrant = false;
 	}
diff --git a/Engine/Source/Editor/Kismet/Private/SSCSEditor.cpp b/Engine/Source/Editor/Kismet/Private/SSCSEditor.cpp
--- a/Engine/Source/Editor/Kismet/Private/SSCSEditor.cpp
+++ b/Engine/Source/Editor/Kismet/Private/SSCSEditor.cpp
@@ -2036,9 +2036,9 @@ void SSCSEditor::Construct( const FArguments& InArgs, TSharedPtr<FBlueprintEdito
 	}
 }
 
-FReply SSCSEditor::OnKeyDown( const FGeometry& MyGeometry, const FKeyboardEvent& InKeyboardEvent )
+FReply SSCSEditor::OnKeyDown( const FGeometry& MyGeometry, const FKeyEvent& InKeyEvent )
 {
-	if ( CommandList->ProcessCommandBindings( InKeyboardEvent ) )
+	if ( CommandList->ProcessCommandBindings( InKeyEvent ) )
 	{
 		return FReply::Handled();
 	}
diff --git a/Engine/Source/Editor/Kismet/Private/SSCSEditor.h b/Engine/Source/Editor/Kismet/Private/SSCSEditor.h
--- a/Engine/Source/Editor/Kismet/Private/SSCSEditor.h
+++ b/Engine/Source/Editor/Kismet/Private/SSCSEditor.h
@@ -379,7 +379,7 @@ public:
 	void Construct(const FArguments& InArgs, TSharedPtr<FBlueprintEditor> InKismet2, USimpleConstructionScript* InSCS, UBlueprint* InBlueprint, TSharedPtr<class SKismetInspector> Inspector = TSharedPtr<class SKismetInspector>() );
 
 	/** Override OnKeyDown */
-	virtual FReply OnKeyDown( const FGeometry& MyGeometry, const FKeyboardEvent& InKeyboardEvent );
+	virtual FReply OnKeyDown( const FGeometry& MyGeometry, const FKeyEvent& InKeyEvent );
 
 	/** Used by tree control - make a widget for a table row from a node */
 	TSharedRef<ITableRow> MakeTableRowWidget(FSCSEditorTreeNodePtrType InNodePtr, const TSharedRef<STableViewBase>& OwnerTable);
diff --git a/Engine/Source/Editor/Kismet/Private/STimelineEditor.cpp b/Engine/Source/Editor/Kismet/Private/STimelineEditor.cpp
--- a/Engine/Source/Editor/Kismet/Private/STimelineEditor.cpp
+++ b/Engine/Source/Editor/Kismet/Private/STimelineEditor.cpp
@@ -1454,9 +1454,9 @@ void STimelineEditor::OnRequestTrackRename() const
 	TrackListView->GetSelectedItems()[0]->OnRenameRequest.Execute();
 }
 
-FReply STimelineEditor::OnKeyDown( const FGeometry& MyGeometry, const FKeyboardEvent& InKeyboardEvent )
+FReply STimelineEditor::OnKeyDown( const FGeometry& MyGeometry, const FKeyEvent& InKeyEvent )
 {
-	if(CommandList->ProcessCommandBindings(InKeyboardEvent))
+	if(CommandList->ProcessCommandBindings(InKeyEvent))
 	{
 		return FReply::Handled();
 	}
diff --git a/Engine/Source/Editor/Kismet/Private/WorkflowOrientedApp/WorkflowTabManager.cpp b/Engine/Source/Editor/Kismet/Private/WorkflowOrientedApp/WorkflowTabManager.cpp
--- a/Engine/Source/Editor/Kismet/Private/WorkflowOrientedApp/WorkflowTabManager.cpp
+++ b/Engine/Source/Editor/Kismet/Private/WorkflowOrientedApp/WorkflowTabManager.cpp
@@ -143,7 +143,7 @@ void FTabInfo::JumpToNearestValidHistoryData()
 			History[CurrentHistoryIndex]->RestoreHistory();
 			History[CurrentHistoryIndex]->GetFactory().Pin()->OnTabActivated(Tab.Pin());
 			FGlobalTabmanager::Get()->SetActiveTab(nullptr);
-			FSlateApplication::Get().ClearKeyboardFocus(EKeyboardFocusCause::SetDirectly);
+			FSlateApplication::Get().ClearKeyboardFocus(EFocusCause::SetDirectly);
 		}
 	}
 }
diff --git a/Engine/Source/Editor/Layers/Private/SLayerBrowser.h b/Engine/Source/Editor/Layers/Private/SLayerBrowser.h
--- a/Engine/Source/Editor/Layers/Private/SLayerBrowser.h
+++ b/Engine/Source/Editor/Layers/Private/SLayerBrowser.h
@@ -48,9 +48,9 @@ public:
 protected:
 
 	/** Overridden from SWidget: Called when a key is pressed down */
-	FReply OnKeyDown( const FGeometry& MyGeometry, const FKeyboardEvent& InKeyboardEvent )
+	FReply OnKeyDown( const FGeometry& MyGeometry, const FKeyEvent& InKeyEvent )
 	{
-		if( InKeyboardEvent.GetKey() == EKeys::Escape && Mode == ELayerBrowserMode::LayerContents )
+		if( InKeyEvent.GetKey() == EKeys::Escape && Mode == ELayerBrowserMode::LayerContents )
 		{
 			SetupLayersMode();			
 			return FReply::Handled();
diff --git a/Engine/Source/Editor/Layers/Private/SLayersView.h b/Engine/Source/Editor/Layers/Private/SLayersView.h
--- a/Engine/Source/Editor/Layers/Private/SLayersView.h
+++ b/Engine/Source/Editor/Layers/Private/SLayersView.h
@@ -132,16 +132,16 @@ protected:
 	}
 
 	/**
-	 * Called after a key is pressed when this widget has keyboard focus
+	 * Called after a key is pressed when this widget has focus
 	 *
 	 * @param MyGeometry The Geometry of the widget receiving the event
-	 * @param  InKeyboardEvent  Keyboard event
+	 * @param  InKeyEvent  Key event
 	 *
 	 * @return  Returns whether the event was handled, along with other possible actions
 	 */
-	virtual FReply OnKeyDown( const FGeometry& MyGeometry, const FKeyboardEvent& InKeyboardEvent ) override
+	virtual FReply OnKeyDown( const FGeometry& MyGeometry, const FKeyEvent& InKeyEvent ) override
 	{
-		return ViewModel->GetCommandList()->ProcessCommandBindings( InKeyboardEvent ) ? FReply::Handled() : FReply::Unhandled();
+		return ViewModel->GetCommandList()->ProcessCommandBindings( InKeyEvent ) ? FReply::Handled() : FReply::Unhandled();
 	}
 
 	/**
diff --git a/Engine/Source/Editor/LevelEditor/Private/DlgDeltaTransform.cpp b/Engine/Source/Editor/LevelEditor/Private/DlgDeltaTransform.cpp
--- a/Engine/Source/Editor/LevelEditor/Private/DlgDeltaTransform.cpp
+++ b/Engine/Source/Editor/LevelEditor/Private/DlgDeltaTransform.cpp
@@ -118,10 +118,10 @@ public:
 	}
 
 	/** Used to intercept Escape key presses, then interprets them as cancel */
-	virtual FReply OnKeyDown( const FGeometry& MyGeometry, const FKeyboardEvent& InKeyboardEvent )
+	virtual FReply OnKeyDown( const FGeometry& MyGeometry, const FKeyEvent& InKeyEvent )
 	{
 		// Pressing escape returns as if the user canceled
-		if ( InKeyboardEvent.GetKey() == EKeys::Escape )
+		if ( InKeyEvent.GetKey() == EKeys::Escape )
 		{
 			return OnButtonClick(FDlgDeltaTransform::Cancel);
 		}
diff --git a/Engine/Source/Editor/LevelEditor/Private/LevelEditor.cpp b/Engine/Source/Editor/LevelEditor/Private/LevelEditor.cpp
--- a/Engine/Source/Editor/LevelEditor/Private/LevelEditor.cpp
+++ b/Engine/Source/Editor/LevelEditor/Private/LevelEditor.cpp
@@ -335,7 +335,7 @@ void FLevelEditorModule::FocusViewport()
 
 		if( FSlateApplication::Get().GeneratePathToWidgetUnchecked( ViewportAsWidget, FocusWidgetPath ) )
 		{
-			FSlateApplication::Get().SetKeyboardFocus( FocusWidgetPath, EKeyboardFocusCause::SetDirectly );
+			FSlateApplication::Get().SetKeyboardFocus( FocusWidgetPath, EFocusCause::SetDirectly );
 		}
 	}
 }
diff --git a/Engine/Source/Editor/LevelEditor/Private/LevelViewportLayout.cpp b/Engine/Source/Editor/LevelEditor/Private/LevelViewportLayout.cpp
--- a/Engine/Source/Editor/LevelEditor/Private/LevelViewportLayout.cpp
+++ b/Engine/Source/Editor/LevelEditor/Private/LevelViewportLayout.cpp
@@ -655,7 +655,7 @@ void FLevelViewportLayout::FinishMaximizeTransition()
 		{
 			// We first need to clear keyboard focus so that Slate doesn't assume that focus won't need to change
 			// simply because the viewport widget object is the same -- it has a new widget path!
-			FSlateApplication::Get().ClearKeyboardFocus( EKeyboardFocusCause::SetDirectly );
+			FSlateApplication::Get().ClearKeyboardFocus( EFocusCause::SetDirectly );
 
 			// Set keyboard focus directly
 			ViewportToFocus->SetKeyboardFocusToThisViewport();
diff --git a/Engine/Source/Editor/LevelEditor/Private/SLevelEditor.cpp b/Engine/Source/Editor/LevelEditor/Private/SLevelEditor.cpp
--- a/Engine/Source/Editor/LevelEditor/Private/SLevelEditor.cpp
+++ b/Engine/Source/Editor/LevelEditor/Private/SLevelEditor.cpp
@@ -1264,9 +1264,9 @@ void SLevelEditor::RefreshEditorModeCommands()
 	}
 }
 
-FReply SLevelEditor::OnKeyDown( const FGeometry& MyGeometry, const FKeyboardEvent& InKeyboardEvent )
+FReply SLevelEditor::OnKeyDown( const FGeometry& MyGeometry, const FKeyEvent& InKeyEvent )
 {
-	// Check to see if any of the actions for the level editor can be processed by the current keyboard event
+	// Check to see if any of the actions for the level editor can be processed by the current event
 	// If we are in debug mode do not process commands
 	if( FSlateApplication::Get().IsNormalExecution() )
 	{
@@ -1308,7 +1308,7 @@ FReply SLevelEditor::OnKeyDown( const FGeometry& MyGeometry, const FKeyboardEven
 		if( ActiveToolkit.IsValid() )
 		{
 			// A toolkit tab is active, so direct all command processing to it
-			if( ActiveToolkit->ProcessCommandBindings( InKeyboardEvent ) )
+			if( ActiveToolkit->ProcessCommandBindings( InKeyEvent ) )
 			{
 				return FReply::Handled();
 			}
@@ -1316,7 +1316,7 @@ FReply SLevelEditor::OnKeyDown( const FGeometry& MyGeometry, const FKeyboardEven
 		else
 		{
 			// No toolkit tab is active, so let the level editor have a chance at the keystroke
-			if( LevelEditorCommands->ProcessCommandBindings( InKeyboardEvent ) )
+			if( LevelEditorCommands->ProcessCommandBindings( InKeyEvent ) )
 			{
 				return FReply::Handled();
 			}
@@ -1326,10 +1326,10 @@ FReply SLevelEditor::OnKeyDown( const FGeometry& MyGeometry, const FKeyboardEven
 	return FReply::Unhandled();
 }
 
-FReply SLevelEditor::OnKeyDownInViewport( const FGeometry& MyGeometry, const FKeyboardEvent& InKeyboardEvent )
+FReply SLevelEditor::OnKeyDownInViewport( const FGeometry& MyGeometry, const FKeyEvent& InKeyEvent )
 {
-	// Check to see if any of the actions for the level editor can be processed by the current keyboard event from a viewport
-	if( LevelEditorCommands->ProcessCommandBindings( InKeyboardEvent ) )
+	// Check to see if any of the actions for the level editor can be processed by the current keyboard from a viewport
+	if( LevelEditorCommands->ProcessCommandBindings( InKeyEvent ) )
 	{
 		return FReply::Handled();
 	}
