Removed some redundant focus switching code from the output log

This was causing the focus to ultimately switch back to the widget that was already focused, however it was also causing IMEs to lose their current context. This force focus isn't needed, as the suggestions list is already set to not take focus when it's spawned.

diff --git a/Engine/Source/Developer/OutputLog/Private/SOutputLog.cpp b/Engine/Source/Developer/OutputLog/Private/SOutputLog.cpp
--- a/Engine/Source/Developer/OutputLog/Private/SOutputLog.cpp
+++ b/Engine/Source/Developer/OutputLog/Private/SOutputLog.cpp
@@ -491,11 +491,6 @@ void SConsoleInputBox::SetSuggestions(TArray<FString>& Elements, bool bInHistory
 		// Ideally if the selection box is open the output window is not changing it's window title (flickers)
 		SuggestionBox->SetIsOpen(true, false);
 		SuggestionListView->RequestScrollIntoView(Suggestions.Last());
-
-		// Force the textbox back into focus.
-		FWidgetPath WidgetToFocusPath;
-		FSlateApplication::Get().GeneratePathToWidgetUnchecked( InputText.ToSharedRef(), WidgetToFocusPath );
-		FSlateApplication::Get().SetKeyboardFocus( WidgetToFocusPath, EFocusCause::SetDirectly );
 	}
 	else
 	{
