Editor: Add LogTimestampMode editor config setting to control how timestamps are displayed in the output log (still defaults to None)

diff --git a/Engine/Source/Developer/OutputLog/Private/SOutputLog.cpp b/Engine/Source/Developer/OutputLog/Private/SOutputLog.cpp
--- a/Engine/Source/Developer/OutputLog/Private/SOutputLog.cpp
+++ b/Engine/Source/Developer/OutputLog/Private/SOutputLog.cpp
@@ -729,6 +729,9 @@ bool SOutputLog::CreateLogMessages( const TCHAR* V, ELogVerbosity::Type Verbosit
 			Style = FName(TEXT("Log.Normal"));
 		}
 
+		// Logging can happen very late during shutdown, even after the UObject system has been torn down, hence the check below
+		ELogTimes::Type LogTimestampMode = UObjectInitialized() ? GetDefault<UEditorStyleSettings>()->LogTimestampMode : ELogTimes::None;
+
 		const int32 OldNumMessages = OutMessages.Num();
 
 		// handle multiline strings by breaking them apart by line
@@ -739,15 +742,15 @@ bool SOutputLog::CreateLogMessages( const TCHAR* V, ELogVerbosity::Type Verbosit
 		bool bIsFirstLineInMessage = true;
 		for (const FTextRange& LineRange : LineRanges)
 		{
-			if (LineRange.IsEmpty())
-				continue;
-
-			FString Line = CurrentLogDump.Mid(LineRange.BeginIndex, LineRange.Len());
-			Line = Line.ConvertTabsToSpaces(4);
+			if (!LineRange.IsEmpty())
+			{
+				FString Line = CurrentLogDump.Mid(LineRange.BeginIndex, LineRange.Len());
+				Line = Line.ConvertTabsToSpaces(4);
 
-			OutMessages.Add(MakeShareable(new FLogMessage(MakeShareable(new FString((bIsFirstLineInMessage) ? FOutputDevice::FormatLogLine(Verbosity, Category, *Line) : Line)), Style)));
+				OutMessages.Add(MakeShareable(new FLogMessage(MakeShareable(new FString((bIsFirstLineInMessage) ? FOutputDevice::FormatLogLine(Verbosity, Category, *Line, LogTimestampMode) : Line)), Style)));
 
-			bIsFirstLineInMessage = false;
+				bIsFirstLineInMessage = false;
+			}
 		}
 
 		return OldNumMessages != OutMessages.Num();
diff --git a/Engine/Source/Editor/EditorStyle/Private/EditorStyleClasses.cpp b/Engine/Source/Editor/EditorStyle/Private/EditorStyleClasses.cpp
--- a/Engine/Source/Editor/EditorStyle/Private/EditorStyleClasses.cpp
+++ b/Engine/Source/Editor/EditorStyle/Private/EditorStyleClasses.cpp
@@ -14,6 +14,7 @@ UEditorStyleSettings::UEditorStyleSettings( const FObjectInitializer& ObjectInit
 	PressedSelectionColor = FLinearColor(0.701f, 0.225f, 0.003f);
 
 	bShowFriendlyNames = true;
+	LogTimestampMode = ELogTimes::None;
 }
 
 
diff --git a/Engine/Source/Editor/EditorStyle/Public/Classes/EditorStyleSettings.h b/Engine/Source/Editor/EditorStyle/Public/Classes/EditorStyleSettings.h
--- a/Engine/Source/Editor/EditorStyle/Public/Classes/EditorStyleSettings.h
+++ b/Engine/Source/Editor/EditorStyle/Public/Classes/EditorStyleSettings.h
@@ -74,6 +74,10 @@ public:
 	UPROPERTY(EditAnywhere, config, Category=UserInterface, AdvancedDisplay)
 	uint32 bExpandConfigurationMenus:1;
 
+	/** The display mode for timestamps in the output log */
+	UPROPERTY(EditAnywhere, config, Category=UserInterface)
+	TEnumAsByte<ELogTimes::Type> LogTimestampMode;
+
 public:
 
 	/**
diff --git a/Engine/Source/Runtime/Core/Public/Misc/OutputDevice.h b/Engine/Source/Runtime/Core/Public/Misc/OutputDevice.h
--- a/Engine/Source/Runtime/Core/Public/Misc/OutputDevice.h
+++ b/Engine/Source/Runtime/Core/Public/Misc/OutputDevice.h
@@ -234,8 +234,13 @@ namespace ELogTimes
 {
 	enum Type
 	{
+		// Do not display log timestamps
 		None,
+
+		// Display log timestamps in UTC
 		UTC,
+
+		// Display log timestamps in seconds elapsed since GStartTime
 		SinceGStartTime
 	};
 }
diff --git a/Engine/Source/Runtime/CoreUObject/Classes/Object.h b/Engine/Source/Runtime/CoreUObject/Classes/Object.h
--- a/Engine/Source/Runtime/CoreUObject/Classes/Object.h
+++ b/Engine/Source/Runtime/CoreUObject/Classes/Object.h
@@ -34,6 +34,25 @@ namespace ESearchDir
 	};
 }
 
+/**
+ * Enum that defines how the log times are to be displayed (mirrored from OutputDevice.h).
+ */
+UENUM()
+namespace ELogTimes
+{
+	enum Type
+	{
+		// Do not display log timestamps
+		None UMETA(DisplayName = "None"),
+
+		// Display log timestamps in UTC
+		UTC UMETA(DisplayName = "UTC"),
+
+		// Display log timestamps in seconds elapsed since GStartTime
+		SinceGStartTime UMETA(DisplayName = "Time since application start"),
+	};
+}
+
 /** Generic axis enum (mirrored for native use in Axis.h). */
 UENUM()
 namespace EAxis
