Editor: Fix Clang error with ? of dissimilar types

diff --git a/Engine/Source/Developer/OutputLog/Private/SOutputLog.cpp b/Engine/Source/Developer/OutputLog/Private/SOutputLog.cpp
--- a/Engine/Source/Developer/OutputLog/Private/SOutputLog.cpp
+++ b/Engine/Source/Developer/OutputLog/Private/SOutputLog.cpp
@@ -729,8 +729,13 @@ bool SOutputLog::CreateLogMessages( const TCHAR* V, ELogVerbosity::Type Verbosit
 			Style = FName(TEXT("Log.Normal"));
 		}
 
-		// Logging can happen very late during shutdown, even after the UObject system has been torn down, hence the check below
-		ELogTimes::Type LogTimestampMode = UObjectInitialized() ? GetDefault<UEditorStyleSettings>()->LogTimestampMode : ELogTimes::None;
+		// Determine how to format timestamps
+		ELogTimes::Type LogTimestampMode = ELogTimes::None;
+		if (UObjectInitialized())
+		{
+			// Logging can happen very late during shutdown, even after the UObject system has been torn down, hence the init check above
+			LogTimestampMode = GetDefault<UEditorStyleSettings>()->LogTimestampMode;
+		}
 
 		const int32 OldNumMessages = OutMessages.Num();
 
