Editor: Fix crash during shutdown with log timestamps

diff --git a/Engine/Source/Developer/OutputLog/Private/SOutputLog.cpp b/Engine/Source/Developer/OutputLog/Private/SOutputLog.cpp
--- a/Engine/Source/Developer/OutputLog/Private/SOutputLog.cpp
+++ b/Engine/Source/Developer/OutputLog/Private/SOutputLog.cpp
@@ -732,8 +732,8 @@ bool SOutputLog::CreateLogMessages( const TCHAR* V, ELogVerbosity::Type Verbosit
 		}
 
 		// Determine how to format timestamps
-		ELogTimes::Type LogTimestampMode = ELogTimes::None;
-		if (UObjectInitialized())
+		static ELogTimes::Type LogTimestampMode = ELogTimes::None;
+		if (UObjectInitialized() && !GExitPurge)
 		{
 			// Logging can happen very late during shutdown, even after the UObject system has been torn down, hence the init check above
 			LogTimestampMode = GetDefault<UEditorStyleSettings>()->LogTimestampMode;
