Made the output log cursor grey so it's easier to see

diff --git a/Engine/Source/Developer/OutputLog/Private/SOutputLog.cpp b/Engine/Source/Developer/OutputLog/Private/SOutputLog.cpp
--- a/Engine/Source/Developer/OutputLog/Private/SOutputLog.cpp
+++ b/Engine/Source/Developer/OutputLog/Private/SOutputLog.cpp
@@ -658,6 +658,7 @@ void SOutputLog::Construct( const FArguments& InArgs )
 	MessagesTextBox = SNew(SMultiLineEditableTextBox)
 		.Style(FEditorStyle::Get(), "Log.TextBox")
 		.TextStyle(FEditorStyle::Get(), "Log.Normal")
+		.ForegroundColor(FLinearColor::Gray)
 		.Marshaller(MessagesTextMarshaller)
 		.IsReadOnly(true)
 		.AlwaysShowScrollbars(true)
