Editor - Now handling previewkeydown in the console command entry widget to prevent navigation and handle key up/down events and tab events correctly as expected for entering autocompeleted console commands.

diff --git a/Engine/Source/Developer/OutputLog/Private/SOutputLog.cpp b/Engine/Source/Developer/OutputLog/Private/SOutputLog.cpp
--- a/Engine/Source/Developer/OutputLog/Private/SOutputLog.cpp
+++ b/Engine/Source/Developer/OutputLog/Private/SOutputLog.cpp
@@ -382,7 +382,7 @@ void SConsoleInputBox::OnTextCommitted( const FText& InText, ETextCommit::Type C
 	}
 }
 
-FReply SConsoleInputBox::OnKeyDown( const FGeometry& MyGeometry, const FKeyEvent& KeyEvent )
+FReply SConsoleInputBox::OnPreviewKeyDown(const FGeometry& MyGeometry, const FKeyEvent& KeyEvent)
 {
 	if(SuggestionBox->IsOpen())
 	{
diff --git a/Engine/Source/Developer/OutputLog/Private/SOutputLog.h b/Engine/Source/Developer/OutputLog/Private/SOutputLog.h
--- a/Engine/Source/Developer/OutputLog/Private/SOutputLog.h
+++ b/Engine/Source/Developer/OutputLog/Private/SOutputLog.h
@@ -67,7 +67,7 @@ protected:
 	virtual bool SupportsKeyboardFocus() const override { return true; }
 
 	// e.g. Tab or Key_Up
-	virtual FReply OnKeyDown( const FGeometry& MyGeometry, const FKeyEvent& KeyEvent ) override;
+	virtual FReply OnPreviewKeyDown( const FGeometry& MyGeometry, const FKeyEvent& KeyEvent ) override;
 
 	void OnFocusLost( const FFocusEvent& InFocusEvent ) override;
 
