Device output log window, currently implemented only for Android. Hidden under editor experimental settings

diff --git a/Engine/Source/Developer/Android/AndroidTargetPlatform/Private/AndroidTargetDevice.h b/Engine/Source/Developer/Android/AndroidTargetPlatform/Private/AndroidTargetDevice.h
--- a/Engine/Source/Developer/Android/AndroidTargetPlatform/Private/AndroidTargetDevice.h
+++ b/Engine/Source/Developer/Android/AndroidTargetPlatform/Private/AndroidTargetDevice.h
@@ -17,6 +17,10 @@ typedef TSharedPtr<class FAndroidTargetDevice, ESPMode::ThreadSafe> FAndroidTarg
  */
 typedef TSharedRef<class FAndroidTargetDevice, ESPMode::ThreadSafe> FAndroidTargetDeviceRef;
 
+/**
+ * Type definition for shared references to instances of FAndroidTargetDevice.
+ */
+typedef TSharedPtr<class FAndroidTargetDeviceOutput, ESPMode::ThreadSafe> FAndroidTargetDeviceOutputPtr;
 
 /**
  * Implements a Android target device.
@@ -97,6 +101,11 @@ public:
 		Model = InNodel;
 	}
 
+	FString GetSerialNumber() const
+	{
+		return SerialNumber;
+	}
+
 public:
 
 	//~ Begin ITargetDevice Interface
@@ -162,8 +171,13 @@ public:
 	virtual bool TerminateProcess(const int32 ProcessId) override;
 	virtual void SetUserCredentials(const FString& UserName, const FString& UserPassword) override;
 	virtual bool GetUserCredentials(FString& OutUserName, FString& OutUserPassword) override;
+	virtual void ExecuteConsoleCommand(const FString& ExecCommand) const override;
+	virtual ITargetDeviceOutputPtr CreateDeviceOutputRouter(FOutputDevice* Output) const override;
 	//~ End ITargetDevice Interface
 
+	/** Full filename for ADB executable. */
+	static bool GetAdbFullFilename(FString& OutFilename);
+
 protected:
 
 	/**
diff --git a/Engine/Source/Developer/Android/AndroidTargetPlatform/Private/AndroidTargetDevice.inl b/Engine/Source/Developer/Android/AndroidTargetPlatform/Private/AndroidTargetDevice.inl
--- a/Engine/Source/Developer/Android/AndroidTargetPlatform/Private/AndroidTargetDevice.inl
+++ b/Engine/Source/Developer/Android/AndroidTargetPlatform/Private/AndroidTargetDevice.inl
@@ -1,8 +1,8 @@
 // Copyright 1998-2015 Epic Games, Inc. All Rights Reserved.
 
-
 /* ITargetDevice interface
  *****************************************************************************/
+#include "AndroidTargetDeviceOutput.h"
 
 inline bool FAndroidTargetDevice::Deploy( const FString& SourceFolder, FString& OutAppId )
 {
@@ -114,12 +114,30 @@ inline bool FAndroidTargetDevice::GetUserCredentials( FString& OutUserName, FStr
 	return false;
 }
 
+inline void FAndroidTargetDevice::ExecuteConsoleCommand(const FString& ExecCommand) const
+{
+	FString AdbCommand = FString::Printf(TEXT("shell am broadcast -a android.intent.action.RUN -e cmd \"%s\""), *ExecCommand);
+	ExecuteAdbCommand(AdbCommand, nullptr, nullptr);
+}
+
+inline ITargetDeviceOutputPtr FAndroidTargetDevice::CreateDeviceOutputRouter(FOutputDevice* Output) const
+{
+	FAndroidTargetDeviceOutputPtr DeviceOutputPtr = MakeShareable(new FAndroidTargetDeviceOutput());
+	if (DeviceOutputPtr->Init(*this, Output))
+	{
+		return DeviceOutputPtr;
+	}
+
+	return nullptr;
+}
 
 /* FAndroidTargetDevice implementation
  *****************************************************************************/
 
-inline bool FAndroidTargetDevice::ExecuteAdbCommand( const FString& CommandLine, FString* OutStdOut, FString* OutStdErr ) const
+inline bool FAndroidTargetDevice::GetAdbFullFilename(FString& OutFilename)
 {
+	TOptional<FString> ResultPath;
+
 	// get the SDK binaries folder
 	TCHAR AndroidDirectory[32768] = { 0 };
 
@@ -130,11 +148,22 @@ inline bool FAndroidTargetDevice::ExecuteAdbCommand( const FString& CommandLine,
 	}
 
 #if PLATFORM_WINDOWS
-	FString Filename = FString::Printf(TEXT("%s\\platform-tools\\adb.exe"), AndroidDirectory);
+	OutFilename = FString::Printf(TEXT("%s\\platform-tools\\adb.exe"), AndroidDirectory);
 #else
-	FString Filename = FString::Printf(TEXT("%s/platform-tools/adb"), AndroidDirectory);
+	OutFilename = FString::Printf(TEXT("%s/platform-tools/adb"), AndroidDirectory);
 #endif
 
+	return true;
+}
+
+inline bool FAndroidTargetDevice::ExecuteAdbCommand( const FString& CommandLine, FString* OutStdOut, FString* OutStdErr ) const
+{
+	FString Filename;
+	if (!GetAdbFullFilename(Filename))
+	{
+		return false;
+	}
+	
 	// execute the command
 	int32 ReturnCode;
 	FString DefaultError;
diff --git a/Engine/Source/Developer/Android/AndroidTargetPlatform/Private/AndroidTargetDeviceOutput.h b/Engine/Source/Developer/Android/AndroidTargetPlatform/Private/AndroidTargetDeviceOutput.h
new file mode 100644
--- /dev/null
+++ b/Engine/Source/Developer/Android/AndroidTargetPlatform/Private/AndroidTargetDeviceOutput.h
@@ -0,0 +1,46 @@
+// Copyright 1998-2015 Epic Games, Inc. All Rights Reserved.
+
+#pragma once
+
+class FAndroidDeviceOutputReaderRunnable : public FRunnable
+{
+public:
+	FAndroidDeviceOutputReaderRunnable(const FString& AdbFilename, const FString& DeviceSerialNumber, FOutputDevice* Output);
+	
+	// FRunnable interface.
+	virtual bool Init(void) override;
+	virtual void Exit(void) override; 
+	virtual void Stop(void) override;
+	virtual uint32 Run(void) override;
+
+private:
+	bool StartLogcatProcess(void);
+
+private:
+	// > 0 if we've been asked to abort work in progress at the next opportunity
+	FThreadSafeCounter	StopTaskCounter;
+	
+	FString				AdbFilename;
+	FString				DeviceSerialNumber;
+	FOutputDevice*		Output;
+	
+	void*				LogcatReadPipe;
+	void*				LogcatWritePipe;
+	FProcHandle			LogcatProcHandle;
+};
+
+/**
+ * Implements a Android target device.
+ */
+class FAndroidTargetDeviceOutput : public ITargetDeviceOutput
+{
+public:
+	bool Init(const FAndroidTargetDevice& TargetDevice, FOutputDevice* Output);
+	
+private:
+	TUniquePtr<FRunnableThread>						DeviceOutputThread;
+	FString											DeviceSerialNumber;
+	FString											DeviceName;
+};
+
+#include "AndroidTargetDeviceOutput.inl"
\ No newline at end of file
diff --git a/Engine/Source/Developer/Android/AndroidTargetPlatform/Private/AndroidTargetDeviceOutput.inl b/Engine/Source/Developer/Android/AndroidTargetPlatform/Private/AndroidTargetDeviceOutput.inl
new file mode 100644
--- /dev/null
+++ b/Engine/Source/Developer/Android/AndroidTargetPlatform/Private/AndroidTargetDeviceOutput.inl
@@ -0,0 +1,100 @@
+// Copyright 1998-2015 Epic Games, Inc. All Rights Reserved.
+
+FAndroidDeviceOutputReaderRunnable::FAndroidDeviceOutputReaderRunnable(const FString& InAdbFilename, const FString& InDeviceSerialNumber, FOutputDevice* InOutput) 
+	: StopTaskCounter(0)
+	, AdbFilename(InAdbFilename)
+	, DeviceSerialNumber(InDeviceSerialNumber)
+	, Output(InOutput)
+	, LogcatReadPipe(nullptr)
+	, LogcatWritePipe(nullptr)
+{
+}
+
+inline bool FAndroidDeviceOutputReaderRunnable::StartLogcatProcess(void)
+{
+	FString Params = FString::Printf(TEXT(" -s %s logcat UE4:V DEBUG:V *:S -v tag"), *DeviceSerialNumber);
+	LogcatProcHandle = FPlatformProcess::CreateProc(*AdbFilename, *Params, true, false, false, NULL, 0, NULL, LogcatWritePipe);
+	return LogcatProcHandle.IsValid();
+}
+
+inline bool FAndroidDeviceOutputReaderRunnable::Init(void) 
+{ 
+	FPlatformProcess::CreatePipe(LogcatReadPipe, LogcatWritePipe);
+	return StartLogcatProcess();
+}
+
+inline void FAndroidDeviceOutputReaderRunnable::Exit(void) 
+{
+	if (LogcatProcHandle.IsValid())
+	{
+		FPlatformProcess::CloseProc(LogcatProcHandle);
+	}
+	FPlatformProcess::ClosePipe(LogcatReadPipe, LogcatWritePipe);
+}
+
+inline void FAndroidDeviceOutputReaderRunnable::Stop(void)
+{
+	StopTaskCounter.Increment();
+}
+
+inline uint32 FAndroidDeviceOutputReaderRunnable::Run(void)
+{
+	while (StopTaskCounter.GetValue() == 0 && LogcatProcHandle.IsValid())
+	{
+		if (!FPlatformProcess::IsProcRunning(LogcatProcHandle))
+		{
+			// When user plugs out USB cable adb process stops
+			// Keep trying to restore adb connection until code that uses this object will not kill us
+			Output->Serialize(TEXT("Trying to restore connection to device..."), ELogVerbosity::Log, NAME_None);
+			FPlatformProcess::CloseProc(LogcatProcHandle);
+			if (StartLogcatProcess())
+			{
+				FPlatformProcess::Sleep(1.0f);
+			}
+			else
+			{
+				Output->Serialize(TEXT("Failed to start adb proccess"), ELogVerbosity::Log, NAME_None);
+				Stop();
+			}
+		}
+		else
+		{
+			FString LogcatOutput = FPlatformProcess::ReadPipe(LogcatReadPipe);
+			if (LogcatOutput.Len() > 0)
+			{
+				TArray<FString> OutputLines;
+				LogcatOutput.ParseIntoArray(OutputLines, TEXT("\n"), false);
+			
+				for (int32 i = 0; i < OutputLines.Num(); ++i)
+				{
+					Output->Serialize(*OutputLines[i], ELogVerbosity::Log, NAME_None);
+				}
+			}
+			
+			FPlatformProcess::Sleep(0.1f);
+		}
+	}
+
+	return 0;
+}
+
+inline bool FAndroidTargetDeviceOutput::Init(const FAndroidTargetDevice& TargetDevice, FOutputDevice* Output)
+{
+	check(Output);
+	// Output will be produced by background thread
+	check(Output->CanBeUsedOnAnyThread());
+	
+	DeviceSerialNumber = TargetDevice.GetSerialNumber();
+	DeviceName = TargetDevice.GetName();
+		
+	FString AdbFilename;
+	if (FAndroidTargetDevice::GetAdbFullFilename(AdbFilename))
+	{
+		auto* Runnable = new FAndroidDeviceOutputReaderRunnable(AdbFilename, DeviceSerialNumber, Output);
+		DeviceOutputThread = TUniquePtr<FRunnableThread>(FRunnableThread::Create(Runnable, TEXT("FAndroidDeviceOutputReaderRunnable")));
+		return true;
+	}
+		
+	return false;
+}
+	
\ No newline at end of file
diff --git a/Engine/Source/Developer/Android/AndroidTargetPlatform/Private/AndroidTargetPlatformPrivatePCH.h b/Engine/Source/Developer/Android/AndroidTargetPlatform/Private/AndroidTargetPlatformPrivatePCH.h
--- a/Engine/Source/Developer/Android/AndroidTargetPlatform/Private/AndroidTargetPlatformPrivatePCH.h
+++ b/Engine/Source/Developer/Android/AndroidTargetPlatform/Private/AndroidTargetPlatformPrivatePCH.h
@@ -30,3 +30,4 @@
 
 #include "AndroidTargetDevice.h"
 #include "AndroidTargetPlatform.h"
+#include "AndroidTargetDeviceOutput.h"
diff --git a/Engine/Source/Developer/OutputLog/OutputLog.Build.cs b/Engine/Source/Developer/OutputLog/OutputLog.Build.cs
--- a/Engine/Source/Developer/OutputLog/OutputLog.Build.cs
+++ b/Engine/Source/Developer/OutputLog/OutputLog.Build.cs
@@ -16,6 +16,8 @@ public class OutputLog : ModuleRules
 				"Slate",
 				"SlateCore",
                 "EditorStyle",
+                "TargetPlatform",
+                "DesktopPlatform",
 			}
 		);
 	}
diff --git a/Engine/Source/Developer/OutputLog/Private/OutputLogModule.cpp b/Engine/Source/Developer/OutputLog/Private/OutputLogModule.cpp
--- a/Engine/Source/Developer/OutputLog/Private/OutputLogModule.cpp
+++ b/Engine/Source/Developer/OutputLog/Private/OutputLogModule.cpp
@@ -3,6 +3,7 @@
 #include "OutputLogPrivatePCH.h"
 #include "SDebugConsole.h"
 #include "SOutputLog.h"
+#include "SDeviceOutputLog.h"
 #include "Editor/WorkspaceMenuStructure/Public/WorkspaceMenuStructureModule.h"
 #include "SDockTab.h"
 
@@ -11,6 +12,7 @@ IMPLEMENT_MODULE( FOutputLogModule, OutputLog );
 namespace OutputLogModule
 {
 	static const FName OutputLogTabName = FName(TEXT("OutputLog"));
+	static const FName DeviceOutputLogTabName = FName(TEXT("DeviceOutputLog"));
 }
 
 /** This class is to capture all log output even if the log window is closed */
@@ -67,6 +69,17 @@ TSharedRef<SDockTab> SpawnOutputLog( const FSpawnTabArgs& Args )
 		];
 }
 
+TSharedRef<SDockTab> SpawnDeviceOutputLog( const FSpawnTabArgs& Args )
+{
+	return SNew(SDockTab)
+		.Icon(FEditorStyle::GetBrush("Log.TabIcon"))
+		.TabRole( ETabRole::NomadTab )
+		.Label( NSLOCTEXT("OutputLog", "DeviceTabTitle", "Device Output Log") )
+		[
+			SNew(SDeviceOutputLog)
+		];
+}
+
 void FOutputLogModule::StartupModule()
 {
 	FGlobalTabmanager::Get()->RegisterNomadTabSpawner(OutputLogModule::OutputLogTabName, FOnSpawnTab::CreateStatic( &SpawnOutputLog ) )
@@ -74,6 +87,13 @@ void FOutputLogModule::StartupModule()
 		.SetTooltipText(NSLOCTEXT("UnrealEditor", "OutputLogTooltipText", "Open the Output Log tab."))
 		.SetGroup( WorkspaceMenu::GetMenuStructure().GetDeveloperToolsLogCategory() )
 		.SetIcon( FSlateIcon(FEditorStyle::GetStyleSetName(), "Log.TabIcon") );
+
+	FGlobalTabmanager::Get()->RegisterNomadTabSpawner(OutputLogModule::DeviceOutputLogTabName, FOnSpawnTab::CreateStatic( &SpawnDeviceOutputLog ) )
+		.SetDisplayName(NSLOCTEXT("UnrealEditor", "DeviceOutputLogTab", "Device Output Log"))
+		.SetTooltipText(NSLOCTEXT("UnrealEditor", "DeviceOutputLogTooltipText", "Open the Device Output Log tab."))
+		.SetGroup( WorkspaceMenu::GetMenuStructure().GetDeveloperToolsLogCategory() )
+		.SetIcon( FSlateIcon(FEditorStyle::GetStyleSetName(), "Log.TabIcon") )
+		.SetAutoGenerateMenuEntry(false); // remove once not Experimental
 	
 	OutputLogHistory = MakeShareable(new FOutputLogHistory);
 }
@@ -83,6 +103,7 @@ void FOutputLogModule::ShutdownModule()
 	if (FSlateApplication::IsInitialized())
 	{
 		FGlobalTabmanager::Get()->UnregisterNomadTabSpawner(OutputLogModule::OutputLogTabName);
+		FGlobalTabmanager::Get()->UnregisterNomadTabSpawner(OutputLogModule::DeviceOutputLogTabName);
 	}
 }
 
diff --git a/Engine/Source/Developer/OutputLog/Private/SDeviceOutputLog.cpp b/Engine/Source/Developer/OutputLog/Private/SDeviceOutputLog.cpp
new file mode 100644
--- /dev/null
+++ b/Engine/Source/Developer/OutputLog/Private/SDeviceOutputLog.cpp
@@ -0,0 +1,329 @@
+// Copyright 1998-2015 Epic Games, Inc. All Rights Reserved.
+
+#include "OutputLogPrivatePCH.h"
+#include "SDeviceOutputLog.h"
+#include "ITargetPlatform.h"
+#include "ITargetPlatformManagerModule.h"
+#include "PlatformInfo.h"
+
+static bool IsSupportedPlatform(ITargetPlatform* Platform)
+{
+	static const FName AndroidPlaftomName("Android"); // TODO: currently implemented only for Android 
+	
+	check(Platform);
+	const auto& PlatfromInfo = Platform->GetPlatformInfo();
+	return PlatfromInfo.IsVanilla() && PlatfromInfo.VanillaPlatformName == AndroidPlaftomName;
+}
+
+
+void SDeviceOutputLog::Construct( const FArguments& InArgs )
+{
+	MessagesTextMarshaller = FOutputLogTextLayoutMarshaller::Create(TArray<TSharedPtr<FLogMessage>>());
+
+	MessagesTextBox = SNew(SMultiLineEditableTextBox)
+		.Style(FEditorStyle::Get(), "Log.TextBox")
+		.TextStyle(FEditorStyle::Get(), "Log.Normal")
+		.ForegroundColor(FLinearColor::Gray)
+		.Marshaller(MessagesTextMarshaller)
+		.IsReadOnly(true)
+		.AlwaysShowScrollbars(true)
+		.OnVScrollBarUserScrolled(this, &SDeviceOutputLog::OnUserScrolled)
+		.ContextMenuExtender(this, &SDeviceOutputLog::ExtendTextBoxMenu);
+
+	ChildSlot
+	[
+		SNew(SVerticalBox)
+
+			// Output log area
+			+SVerticalBox::Slot()
+			.FillHeight(1)
+			[
+				MessagesTextBox.ToSharedRef()
+			]
+			// The console input box
+			+SVerticalBox::Slot()
+			.AutoHeight()
+			.Padding(FMargin(0.0f, 4.0f, 0.0f, 0.0f))
+			[
+				SNew(SHorizontalBox)
+
+				+SHorizontalBox::Slot()
+				.AutoWidth()
+				[
+					SAssignNew(TargetDeviceComboBox, SComboBox<FTargetDeviceEntryPtr>)
+					.OptionsSource(&DeviceList)
+					.ButtonStyle(FEditorStyle::Get(), "ToolBar.Button")
+					.OnSelectionChanged(this, &SDeviceOutputLog::OnDeviceSelectionChanged)
+					.OnGenerateWidget(this, &SDeviceOutputLog::GenerateWidgetForDeviceComboBox)
+					.ContentPadding(FMargin(4.0f, 0.0f))
+					.Content()
+					[
+						SNew(SHorizontalBox)
+						+SHorizontalBox::Slot()
+						.AutoWidth()
+						[
+							SNew(SBox)
+							.WidthOverride(16)
+							.HeightOverride(16)
+							[
+								SNew(SImage).Image(this, &SDeviceOutputLog::GetSelectedTargetDeviceBrush)
+							]
+						]
+			
+						+SHorizontalBox::Slot()
+						.VAlign(VAlign_Center)
+						[
+							SNew(STextBlock).Text(this, &SDeviceOutputLog::GetSelectedTargetDeviceText)
+						]
+					]
+				]
+
+				+SHorizontalBox::Slot()
+				.Padding(FMargin(4.0f, 0.0f, 0.0f, 4.0f))
+				.FillWidth(1)
+				[
+					SNew(SConsoleInputBox)
+					.ConsoleCommandCustomExec(this, &SDeviceOutputLog::ExecuteConsoleCommand)
+					.OnConsoleCommandExecuted(this, &SDeviceOutputLog::OnConsoleCommandExecuted)
+					// Always place suggestions above the input line for the output log widget
+					.SuggestionListPlacement( MenuPlacement_AboveAnchor )
+				]
+			]
+	];
+
+	bIsUserScrolled = false;
+	RequestForceScroll();
+	
+	//
+	TArray<ITargetPlatform*> Platforms = GetTargetPlatformManager()->GetTargetPlatforms();
+	for (ITargetPlatform* Platform : Platforms)
+	{
+		if (IsSupportedPlatform(Platform))
+		{
+			Platform->OnDeviceDiscovered().AddRaw(this, &SDeviceOutputLog::HandleTargetPlatformDeviceDiscovered);
+			Platform->OnDeviceLost().AddRaw(this, &SDeviceOutputLog::HandleTargetPlatformDeviceLost);
+		}
+	}
+		
+	// Get list of available devices
+	for (ITargetPlatform* Platform : Platforms)
+	{
+		if (IsSupportedPlatform(Platform))
+		{
+			TArray<ITargetDevicePtr> TargetDevices;
+			Platform->GetAllDevices(TargetDevices);
+
+			for (const ITargetDevicePtr& Device : TargetDevices)
+			{
+				if (Device.IsValid())
+				{
+					AddDeviceEntry(Device.ToSharedRef());
+				}
+			}
+		}
+	}
+}
+
+SDeviceOutputLog::~SDeviceOutputLog()
+{
+	ITargetPlatformManagerModule* Module = FModuleManager::GetModulePtr<ITargetPlatformManagerModule>("TargetPlatform");
+	if (Module)
+	{
+		TArray<ITargetPlatform*> Platforms = Module->GetTargetPlatforms();
+		for (ITargetPlatform* Platform : Platforms)
+		{
+			Platform->OnDeviceDiscovered().RemoveAll(this);
+			Platform->OnDeviceLost().RemoveAll(this);
+		}
+	}
+}
+
+void SDeviceOutputLog::Tick(const FGeometry& AllottedGeometry, const double InCurrentTime, const float InDeltaTime)
+{
+	FScopeLock ScopeLock(&BufferedLinesSynch);
+	if (BufferedLines.Num() > 0)
+	{
+		for (const FBufferedLine& Line : BufferedLines)
+		{
+			MessagesTextMarshaller->AppendMessage(*Line.Data, Line.Verbosity, Line.Category);
+		}
+
+		// Don't scroll to the bottom automatically when the user is scrolling the view or has scrolled it away from the bottom.
+		if (!bIsUserScrolled)
+		{
+			MessagesTextBox->ScrollTo(FTextLocation(MessagesTextMarshaller->GetNumMessages() - 1));
+		}
+
+		BufferedLines.Empty(32);
+	}
+}
+
+void SDeviceOutputLog::Serialize(const TCHAR* V, ELogVerbosity::Type Verbosity, const class FName& Category)
+{
+	FScopeLock ScopeLock(&BufferedLinesSynch);
+	BufferedLines.Add(FBufferedLine(V, Category, Verbosity));
+}
+
+bool SDeviceOutputLog::CanBeUsedOnAnyThread() const
+{
+	return true;
+}
+
+void SDeviceOutputLog::ExecuteConsoleCommand(const FString& ExecCommand)
+{
+	FTargetDeviceEntryPtr SelectedDeviceEntry = TargetDeviceComboBox->GetSelectedItem();
+	if (SelectedDeviceEntry.IsValid())
+	{
+		ITargetDevicePtr PinnedPtr = SelectedDeviceEntry->DeviceWeakPtr.Pin();
+		if (PinnedPtr.IsValid())
+		{
+			PinnedPtr->ExecuteConsoleCommand(ExecCommand);
+		}
+	}
+}
+
+void SDeviceOutputLog::HandleTargetPlatformDeviceLost(ITargetDeviceRef LostDevice)
+{
+	auto LostDeviceId = LostDevice->GetId();
+	FTargetDeviceEntryPtr SelectedDeviceEntry = TargetDeviceComboBox->GetSelectedItem();
+
+	if (SelectedDeviceEntry.IsValid() && SelectedDeviceEntry->DeviceId == LostDeviceId)
+	{
+		// Kill device output object, but do not clean up output in the window
+		CurrentDeviceOutputPtr.Reset();
+	}
+	
+	// Should not do it, but what if someone somewhere holds strong reference to a lost device?
+	for (const TSharedPtr<FTargetDeviceEntry>& EntryPtr : DeviceList)
+	{
+		if (EntryPtr->DeviceId == LostDeviceId)
+		{
+			EntryPtr->DeviceWeakPtr = nullptr;
+		}
+	}
+}
+
+void SDeviceOutputLog::HandleTargetPlatformDeviceDiscovered(ITargetDeviceRef DiscoveredDevice)
+{
+	FTargetDeviceId DiscoveredDeviceId = DiscoveredDevice->GetId();
+
+	int32 ExistingEntryIdx = DeviceList.IndexOfByPredicate([&](const TSharedPtr<FTargetDeviceEntry>& Other) {
+		return (Other->DeviceId == DiscoveredDeviceId);
+	});
+
+	if (DeviceList.IsValidIndex(ExistingEntryIdx))
+	{
+		DeviceList[ExistingEntryIdx]->DeviceWeakPtr = DiscoveredDevice;
+		
+		if (TargetDeviceComboBox->GetSelectedItem() == DeviceList[ExistingEntryIdx])
+		{
+			CurrentDeviceOutputPtr = DiscoveredDevice->CreateDeviceOutputRouter(this);
+		}
+	}
+	else
+	{
+		AddDeviceEntry(DiscoveredDevice);
+		TargetDeviceComboBox->RefreshOptions();
+	}
+}
+
+void SDeviceOutputLog::AddDeviceEntry(ITargetDeviceRef TargetDevice)
+{
+	using namespace PlatformInfo;
+	FName DeviceIconStyleName = TargetDevice->GetTargetPlatform().GetPlatformInfo().GetIconStyleName(EPlatformIconSize::Normal);
+	
+	TSharedPtr<FTargetDeviceEntry> DeviceEntry = MakeShareable(new FTargetDeviceEntry());
+	
+	DeviceEntry->DeviceId = TargetDevice->GetId();
+	DeviceEntry->DeviceName = TargetDevice->GetName();
+	DeviceEntry->DeviceIconBrush = FEditorStyle::GetBrush(DeviceIconStyleName);
+	DeviceEntry->DeviceWeakPtr = TargetDevice;
+	
+	DeviceList.Add(DeviceEntry);
+}
+
+void SDeviceOutputLog::OnDeviceSelectionChanged(FTargetDeviceEntryPtr DeviceEntry, ESelectInfo::Type SelectInfo)
+{
+	CurrentDeviceOutputPtr.Reset();
+	OnClearLog();
+	
+	if (DeviceEntry.IsValid())
+	{
+		ITargetDevicePtr PinnedPtr = DeviceEntry->DeviceWeakPtr.Pin();
+		if (PinnedPtr.IsValid() && PinnedPtr->IsConnected())
+		{
+			CurrentDeviceOutputPtr = PinnedPtr->CreateDeviceOutputRouter(this);
+		}
+	}
+}
+
+TSharedRef<SWidget> SDeviceOutputLog::GenerateWidgetForDeviceComboBox(FTargetDeviceEntryPtr DeviceEntry) const
+{
+	return 
+		SNew(SBox)
+		[
+			SNew(SHorizontalBox)
+
+			+SHorizontalBox::Slot()
+			.AutoWidth()
+			[
+				SNew(SBox)
+				.WidthOverride(24)
+				.HeightOverride(24)
+				[
+					SNew(SImage).Image(GetTargetDeviceBrush(DeviceEntry))
+				]
+			]
+			
+			+SHorizontalBox::Slot()
+			.VAlign(VAlign_Center)
+			.Padding(FMargin(4.0f, 0.0f))
+			[
+				SNew(STextBlock).Text(this, &SDeviceOutputLog::GetTargetDeviceText, DeviceEntry)
+			]
+		];
+}
+
+const FSlateBrush* SDeviceOutputLog::GetTargetDeviceBrush(FTargetDeviceEntryPtr DeviceEntry) const
+{
+	if (DeviceEntry.IsValid())
+	{
+		return DeviceEntry->DeviceIconBrush;
+	}
+	else
+	{
+		return FEditorStyle::GetBrush("Launcher.Instance_Unknown");
+	}
+}
+
+const FSlateBrush* SDeviceOutputLog::GetSelectedTargetDeviceBrush() const
+{
+	FTargetDeviceEntryPtr DeviceEntry = TargetDeviceComboBox->GetSelectedItem();
+	return GetTargetDeviceBrush(DeviceEntry);
+}
+
+FText SDeviceOutputLog::GetTargetDeviceText(FTargetDeviceEntryPtr DeviceEntry) const
+{
+	if (DeviceEntry.IsValid())
+	{
+		ITargetDevicePtr PinnedPtr = DeviceEntry->DeviceWeakPtr.Pin();
+		if (PinnedPtr.IsValid() && PinnedPtr->IsConnected())
+		{
+			return FText::FromString(DeviceEntry->DeviceName);
+		}
+		else
+		{
+			return FText::Format(NSLOCTEXT("OutputLog", "TargetDeviceOffline", "{0} (Offline)"), FText::FromString(DeviceEntry->DeviceName));
+		}
+	}
+	else
+	{
+		return NSLOCTEXT("OutputLog", "UnknownTargetDevice", "<Unknown device>");
+	}
+}
+
+FText SDeviceOutputLog::GetSelectedTargetDeviceText() const
+{
+	FTargetDeviceEntryPtr DeviceEntry = TargetDeviceComboBox->GetSelectedItem();
+	return GetTargetDeviceText(DeviceEntry);
+}
diff --git a/Engine/Source/Developer/OutputLog/Private/SDeviceOutputLog.h b/Engine/Source/Developer/OutputLog/Private/SDeviceOutputLog.h
new file mode 100644
--- /dev/null
+++ b/Engine/Source/Developer/OutputLog/Private/SDeviceOutputLog.h
@@ -0,0 +1,72 @@
+// Copyright 1998-2015 Epic Games, Inc. All Rights Reserved.
+
+#pragma once
+
+#include "SOutputLog.h"
+#include "TargetDeviceId.h"
+#include "ITargetDevice.h"
+
+struct FTargetDeviceEntry
+{
+	FTargetDeviceId			DeviceId;
+	FString					DeviceName;
+	const FSlateBrush*		DeviceIconBrush;
+	ITargetDeviceWeakPtr	DeviceWeakPtr;
+};
+
+typedef TSharedPtr<FTargetDeviceEntry> FTargetDeviceEntryPtr;
+
+class SDeviceOutputLog : public SOutputLog
+{
+public:
+	SLATE_BEGIN_ARGS( SDeviceOutputLog )
+	{}
+	SLATE_END_ARGS()
+
+	/** Destructor for output log, so we can unregister from notifications */
+	~SDeviceOutputLog();
+
+	/**
+	 * Construct this widget.  Called by the SNew() Slate macro.
+	 *
+	 * @param	InArgs	Declaration used by the SNew() macro to construct this widget
+	 */
+	void Construct( const FArguments& InArgs );
+
+protected:
+	// SWidget interface
+	virtual void Tick(const FGeometry& AllottedGeometry, const double InCurrentTime, const float InDeltaTime) override;
+	
+	// FOutputDevice interface
+	virtual void Serialize( const TCHAR* V, ELogVerbosity::Type Verbosity, const class FName& Category ) override;
+	virtual bool CanBeUsedOnAnyThread() const;
+
+	void ExecuteConsoleCommand(const FString& ExecCommand);
+
+	/** Callback for lost target devices. */
+	void HandleTargetPlatformDeviceLost(ITargetDeviceRef LostDevice);
+	/** Callback for discovered target devices. */
+	void HandleTargetPlatformDeviceDiscovered(ITargetDeviceRef DiscoveredDevice);
+
+	void AddDeviceEntry(ITargetDeviceRef TargetDevice);
+
+	void OnDeviceSelectionChanged(FTargetDeviceEntryPtr DeviceEntry, ESelectInfo::Type SelectInfo);
+	TSharedRef<SWidget> GenerateWidgetForDeviceComboBox(FTargetDeviceEntryPtr DeviceEntry) const;
+
+	FText GetTargetDeviceText(FTargetDeviceEntryPtr DeviceEntry) const;
+	FText GetSelectedTargetDeviceText() const;
+	
+	const FSlateBrush* GetTargetDeviceBrush(FTargetDeviceEntryPtr DeviceEntry) const;
+	const FSlateBrush* GetSelectedTargetDeviceBrush() const;
+		
+private:
+	TArray<FTargetDeviceEntryPtr>			DeviceList;
+	ITargetDeviceWeakPtr					CurrentDevicePtr;
+	ITargetDeviceOutputPtr					CurrentDeviceOutputPtr;
+
+	TSharedPtr<SComboBox<FTargetDeviceEntryPtr>> TargetDeviceComboBox;
+		
+	/** Synchronization object for access to buffered lines */
+	FCriticalSection		BufferedLinesSynch;
+	TArray<FBufferedLine>	BufferedLines;
+};
diff --git a/Engine/Source/Developer/OutputLog/Private/SOutputLog.cpp b/Engine/Source/Developer/OutputLog/Private/SOutputLog.cpp
--- a/Engine/Source/Developer/OutputLog/Private/SOutputLog.cpp
+++ b/Engine/Source/Developer/OutputLog/Private/SOutputLog.cpp
@@ -3,7 +3,6 @@
 #include "OutputLogPrivatePCH.h"
 #include "SOutputLog.h"
 #include "SScrollBorder.h"
-#include "BaseTextLayoutMarshaller.h"
 #include "GameFramework/GameMode.h"
 #include "Engine/LocalPlayer.h"
 #include "GameFramework/GameState.h"
@@ -130,6 +129,7 @@ BEGIN_SLATE_FUNCTION_BUILD_OPTIMIZATION
 void SConsoleInputBox::Construct( const FArguments& InArgs )
 {
 	OnConsoleCommandExecuted = InArgs._OnConsoleCommandExecuted;
+	ConsoleCommandCustomExec = InArgs._ConsoleCommandCustomExec;
 
 	ChildSlot
 	[
@@ -314,6 +314,11 @@ void SConsoleInputBox::OnTextCommitted( const FText& InText, ETextCommit::Type C
 			bIgnoreUIUpdate = false;
 			
 			// Exec!
+			if (ConsoleCommandCustomExec.IsBound())
+			{
+				ConsoleCommandCustomExec.Execute(ExecString);
+			}
+			else
 			{
 				bool bWasHandled = false;
 				UWorld* World = NULL;
@@ -534,35 +539,6 @@ FString SConsoleInputBox::GetSelectionText() const
 	return ret;
 }
 
-/** Output log text marshaller to convert an array of FLogMessages into styled lines to be consumed by an FTextLayout */
-class FOutputLogTextLayoutMarshaller : public FBaseTextLayoutMarshaller
-{
-public:
-
-	static TSharedRef< FOutputLogTextLayoutMarshaller > Create(TArray< TSharedPtr<FLogMessage> > InMessages);
-
-	virtual ~FOutputLogTextLayoutMarshaller();
-	
-	// ITextLayoutMarshaller
-	virtual void SetText(const FString& SourceString, FTextLayout& TargetTextLayout) override;
-	virtual void GetText(FString& TargetString, const FTextLayout& SourceTextLayout) override;
-
-	bool AppendMessage(const TCHAR* InText, const ELogVerbosity::Type InVerbosity, const FName& InCategory);
-	void ClearMessages();
-	int32 GetNumMessages() const;
-
-protected:
-
-	FOutputLogTextLayoutMarshaller(TArray< TSharedPtr<FLogMessage> > InMessages);
-
-	void AppendMessageToTextLayout(const TSharedPtr<FLogMessage>& Message);
-
-	/** All log messages to show in the text box */
-	TArray< TSharedPtr<FLogMessage> > Messages;
-
-	FTextLayout* TextLayout;
-};
-
 TSharedRef< FOutputLogTextLayoutMarshaller > FOutputLogTextLayoutMarshaller::Create(TArray< TSharedPtr<FLogMessage> > InMessages)
 {
 	return MakeShareable(new FOutputLogTextLayoutMarshaller(MoveTemp(InMessages)));
diff --git a/Engine/Source/Developer/OutputLog/Private/SOutputLog.h b/Engine/Source/Developer/OutputLog/Private/SOutputLog.h
--- a/Engine/Source/Developer/OutputLog/Private/SOutputLog.h
+++ b/Engine/Source/Developer/OutputLog/Private/SOutputLog.h
@@ -1,11 +1,10 @@
 // Copyright 1998-2015 Epic Games, Inc. All Rights Reserved.
 
 #pragma once
-
+#include "BaseTextLayoutMarshaller.h"
 
 class FOutputLogTextLayoutMarshaller;
 
-
 /**
  * A single log message for the output log, holding a message and
  * a style, for color and bolding of the message.
@@ -31,6 +30,7 @@ class SConsoleInputBox
 {
 
 public:
+	DECLARE_DELEGATE_OneParam(FExecuteConsoleCommand, const FString& /*ExecCommand*/)
 
 	SLATE_BEGIN_ARGS( SConsoleInputBox )
 		: _SuggestionListPlacement( MenuPlacement_BelowAnchor )
@@ -39,6 +39,9 @@ public:
 		/** Where to place the suggestion list */
 		SLATE_ARGUMENT( EMenuPlacement, SuggestionListPlacement )
 
+		/** Custom executor for console command, will be used when bound */
+		SLATE_EVENT( FExecuteConsoleCommand, ConsoleCommandCustomExec)
+
 		/** Called when a console command is executed */
 		SLATE_EVENT( FSimpleDelegate, OnConsoleCommandExecuted )
 	SLATE_END_ARGS()
@@ -89,7 +92,6 @@ protected:
 
 	FString GetSelectionText() const;
 
-
 private:
 
 	/** Editable text widget */
@@ -107,6 +109,9 @@ private:
 	/** Delegate to call when a console command is executed */
 	FSimpleDelegate OnConsoleCommandExecuted;
 
+	/** Delegate to call to execute console command */
+	FExecuteConsoleCommand ConsoleCommandCustomExec;
+
 	/** -1 if not set, otherwise index into Suggestions */
 	int32 SelectedSuggestion;
 
@@ -160,7 +165,7 @@ protected:
 
 	virtual void Serialize( const TCHAR* V, ELogVerbosity::Type Verbosity, const class FName& Category ) override;
 
-private:
+protected:
 	/**
 	 * Extends the context menu used by the text box
 	 */
@@ -196,3 +201,33 @@ private:
 	/** True if the user has scrolled the window upwards */
 	bool bIsUserScrolled;
 };
+
+
+/** Output log text marshaller to convert an array of FLogMessages into styled lines to be consumed by an FTextLayout */
+class FOutputLogTextLayoutMarshaller : public FBaseTextLayoutMarshaller
+{
+public:
+
+	static TSharedRef< FOutputLogTextLayoutMarshaller > Create(TArray< TSharedPtr<FLogMessage> > InMessages);
+
+	virtual ~FOutputLogTextLayoutMarshaller();
+	
+	// ITextLayoutMarshaller
+	virtual void SetText(const FString& SourceString, FTextLayout& TargetTextLayout) override;
+	virtual void GetText(FString& TargetString, const FTextLayout& SourceTextLayout) override;
+
+	bool AppendMessage(const TCHAR* InText, const ELogVerbosity::Type InVerbosity, const FName& InCategory);
+	void ClearMessages();
+	int32 GetNumMessages() const;
+
+protected:
+
+	FOutputLogTextLayoutMarshaller(TArray< TSharedPtr<FLogMessage> > InMessages);
+
+	void AppendMessageToTextLayout(const TSharedPtr<FLogMessage>& Message);
+
+	/** All log messages to show in the text box */
+	TArray< TSharedPtr<FLogMessage> > Messages;
+
+	FTextLayout* TextLayout;
+};
\ No newline at end of file
diff --git a/Engine/Source/Developer/TargetPlatform/Public/Interfaces/ITargetDevice.h b/Engine/Source/Developer/TargetPlatform/Public/Interfaces/ITargetDevice.h
--- a/Engine/Source/Developer/TargetPlatform/Public/Interfaces/ITargetDevice.h
+++ b/Engine/Source/Developer/TargetPlatform/Public/Interfaces/ITargetDevice.h
@@ -190,6 +190,8 @@ typedef TSharedRef<class ITargetDevice, ESPMode::ThreadSafe> ITargetDeviceRef;
 /** Type definition for weak pointers to instances of ITargetDevice. */
 typedef TWeakPtr<class ITargetDevice, ESPMode::ThreadSafe> ITargetDeviceWeakPtr;
 
+/** Type definition for shared pointers to instances of ITargetDeviceOutput. */
+typedef TSharedPtr<class ITargetDeviceOutput, ESPMode::ThreadSafe> ITargetDeviceOutputPtr;
 
 /**
  * Interface for target devices.
@@ -369,6 +371,24 @@ public:
 	 */
 	virtual bool GetUserCredentials( FString& OutUserName, FString& OutUserPassword ) = 0;
 
+	/**
+	 * Execute console command on the device
+	 * 
+	 * @param ExecCommand Command to execute
+	 */
+	virtual void ExecuteConsoleCommand(const FString& ExecCommand) const {};
+
+	/**
+	 * Create device output router
+	 *
+	 * This will route device logs into specified OutputDevice 
+	 * until connection to device is alive
+	 *
+	 * @param Output OutputDevice to where output should be routed (has to be thread safe)
+	 * @return Valid router object for devices that support output routing
+	 */
+	virtual ITargetDeviceOutputPtr CreateDeviceOutputRouter(FOutputDevice* Output) const { return nullptr; };
+
 public:
 	
 	/** Virtual destructor. */
diff --git a/Engine/Source/Developer/TargetPlatform/Public/Interfaces/ITargetDeviceOutput.h b/Engine/Source/Developer/TargetPlatform/Public/Interfaces/ITargetDeviceOutput.h
new file mode 100644
--- /dev/null
+++ b/Engine/Source/Developer/TargetPlatform/Public/Interfaces/ITargetDeviceOutput.h
@@ -0,0 +1,13 @@
+// Copyright 1998-2015 Epic Games, Inc. All Rights Reserved.
+
+#pragma once
+
+/**
+ * Interface for target device output log reader.
+ */
+class ITargetDeviceOutput
+{
+public:
+	/** Virtual destructor. */
+	virtual ~ITargetDeviceOutput() { }
+};
diff --git a/Engine/Source/Developer/TargetPlatform/Public/TargetPlatform.h b/Engine/Source/Developer/TargetPlatform/Public/TargetPlatform.h
--- a/Engine/Source/Developer/TargetPlatform/Public/TargetPlatform.h
+++ b/Engine/Source/Developer/TargetPlatform/Public/TargetPlatform.h
@@ -18,6 +18,7 @@
 #include "ITargetPlatform.h"
 #include "ITargetPlatformModule.h"
 #include "ITargetPlatformManagerModule.h"
+#include "ITargetDeviceOutput.h"
 
 #include "IAudioFormat.h"
 #include "IAudioFormatModule.h"
diff --git a/Engine/Source/Editor/MainFrame/Private/Menus/MainMenu.cpp b/Engine/Source/Editor/MainFrame/Private/Menus/MainMenu.cpp
--- a/Engine/Source/Editor/MainFrame/Private/Menus/MainMenu.cpp
+++ b/Engine/Source/Editor/MainFrame/Private/Menus/MainMenu.cpp
@@ -197,9 +197,10 @@ void FMainMenu::FillWindowMenu( FMenuBuilder& MenuBuilder, const TSharedRef< FEx
 		bool bLocalizationDashboard = GetDefault<UEditorExperimentalSettings>()->bEnableLocalizationDashboard;
 		bool bTranslationPicker = GetDefault<UEditorExperimentalSettings>()->bEnableTranslationPicker;
 		bool bMergeActors = GetDefault<UEditorExperimentalSettings>()->bActorMerging;
+		bool bDeviceOutputLog = GetDefault<UEditorExperimentalSettings>()->bDeviceOutputLog;
 
 		// Make sure at least one is enabled before creating the section
-		if (bMessagingDebugger || bBlutility || bLocalizationDashboard || bTranslationPicker || bMergeActors)
+		if (bMessagingDebugger || bBlutility || bLocalizationDashboard || bTranslationPicker || bMergeActors || bDeviceOutputLog)
 		{
 			MenuBuilder.BeginSection("ExperimentalTabSpawners", LOCTEXT("ExperimentalTabSpawnersHeading", "Experimental"));
 			{
@@ -257,6 +258,18 @@ void FMainMenu::FillWindowMenu( FMenuBuilder& MenuBuilder, const TSharedRef< FEx
 						FUIAction(FExecuteAction::CreateStatic(&FMainMenu::OpenMergeActors))
 						);
 				}
+
+				// Device output log
+				if (bDeviceOutputLog)
+				{
+					MenuBuilder.AddMenuEntry(
+						LOCTEXT("DeviceOutputLogMenuLabel", "Device Output Log"),
+						LOCTEXT("DeviceOutputLogToolTip", "Open the Device Output Log tab."),
+						FSlateIcon(FEditorStyle::GetStyleSetName(), "Log.TabIcon"),
+						FUIAction(FExecuteAction::CreateStatic(&FMainMenu::OpenDeviceOutputLog))
+						);
+				}
+
 			}
 			MenuBuilder.EndSection();
 		}
diff --git a/Engine/Source/Editor/MainFrame/Private/Menus/MainMenu.h b/Engine/Source/Editor/MainFrame/Private/Menus/MainMenu.h
--- a/Engine/Source/Editor/MainFrame/Private/Menus/MainMenu.h
+++ b/Engine/Source/Editor/MainFrame/Private/Menus/MainMenu.h
@@ -120,4 +120,14 @@ private:
 	{
 		FGlobalTabmanager::Get()->InvokeTab(FName(TEXT("MergeActors")));
 	}
+	
+	
+	/**
+	* Opens the 'Device Output Log' tab.
+	* Remove this when it is no longer experimental.
+	*/
+	static void OpenDeviceOutputLog()
+	{
+		FGlobalTabmanager::Get()->InvokeTab(FName(TEXT("DeviceOutputLog")));
+	}
 };
diff --git a/Engine/Source/Editor/UnrealEd/Classes/Settings/EditorExperimentalSettings.h b/Engine/Source/Editor/UnrealEd/Classes/Settings/EditorExperimentalSettings.h
--- a/Engine/Source/Editor/UnrealEd/Classes/Settings/EditorExperimentalSettings.h
+++ b/Engine/Source/Editor/UnrealEd/Classes/Settings/EditorExperimentalSettings.h
@@ -53,6 +53,10 @@ public:
 	UPROPERTY(EditAnywhere, config, Category = Tools, meta = (DisplayName = "Actor Merging"))
 	bool bActorMerging;
 
+	/** Device output log window (currently implemented for Android only)*/
+	UPROPERTY(EditAnywhere, config, Category = Tools, meta = (DisplayName = "Device Output Log"))
+	bool bDeviceOutputLog;
+
 	/** Specify which console-specific nomenclature to use for gamepad label text */
 	UPROPERTY(EditAnywhere, config, Category=UserInterface, meta=(DisplayName="Console for Gamepad Labels"))
 	TEnumAsByte<EConsoleForGamepadLabels::Type> ConsoleForGamepadLabels;
