Copying //UE4/Dev-Enterprise to //UE4/Dev-Main (Source: //UE4/Dev-Enterprise @ 3769617)

============================
  MAJOR FEATURES & CHANGES
============================

Change 3678071 by JeanLuc.Corenthin
	Fixed a bug in AActor::CheckComponentInstanceName. The implementation was crashing with name containing only digits.

Change 3706350 by [REDACTED_PERSON]

	Fixed broken style from the Python merge

Change 3709830 by [REDACTED_PERSON]

	Exposed FVector4 and FQuat to TBaseStructure

Change 3713856 by [REDACTED_PERSON]

	Added missing FVector4-= operator

Change 3716874 by [REDACTED_PERSON]

	UAT ListThirdPartySoftware: Add support for external project. New parameter -ProjectPath to specifiy path to .uproject so the build rules of the project are parsed and the folder is part of the search for TPS files.

Change 3728856 by [REDACTED_PERSON]

	Get the value of each row for a specific column/property in a data table. It doesn't include the title of that column.

Change 3739050 by [REDACTED_PERSON]

	Exposed FPlane to TBaseStructure

diff --git a/Engine/Source/Developer/MeshUtilities/Private/LayoutUV.cpp b/Engine/Source/Developer/MeshUtilities/Private/LayoutUV.cpp
--- a/Engine/Source/Developer/MeshUtilities/Private/LayoutUV.cpp
+++ b/Engine/Source/Developer/MeshUtilities/Private/LayoutUV.cpp
@@ -22,7 +22,7 @@ FLayoutUV::FLayoutUV( FRawMesh* InMesh, uint32 InSrcChannel, uint32 InDstChannel
 	, LayoutVersion( ELightmapUVVersion::Latest )
 {}
 
-void FLayoutUV::FindCharts( const TMultiMap<int32,int32>& OverlappingCorners )
+int32 FLayoutUV::FindCharts( const TMultiMap<int32,int32>& OverlappingCorners )
 {
 	double Begin = FPlatformTime::Seconds();
 
@@ -550,6 +550,8 @@ void FLayoutUV::FindCharts( const TMultiMap<int32,int32>& OverlappingCorners )
 	double End = FPlatformTime::Seconds();
 
 	UE_LOG(LogLayoutUV, Display, TEXT("FindCharts: %s"), *FPlatformTime::PrettyTime(End - Begin) );
+
+	return Charts.Num();
 }
 
 bool FLayoutUV::FindBestPacking()
diff --git a/Engine/Source/Developer/MeshUtilities/Private/MeshUtilities.cpp b/Engine/Source/Developer/MeshUtilities/Private/MeshUtilities.cpp
--- a/Engine/Source/Developer/MeshUtilities/Private/MeshUtilities.cpp
+++ b/Engine/Source/Developer/MeshUtilities/Private/MeshUtilities.cpp
@@ -5152,6 +5152,19 @@ static bool NonOpaqueMaterialPredicate(UStaticMeshComponent* InMesh)
 }
 
 void FMeshUtilities::RecomputeTangentsAndNormalsForRawMesh(bool bRecomputeTangents, bool bRecomputeNormals, const FMeshBuildSettings& InBuildSettings, FRawMesh &OutRawMesh) const
+{
+	// Compute any missing tangents.
+	if (bRecomputeNormals || bRecomputeTangents)
+	{
+		float ComparisonThreshold = InBuildSettings.bRemoveDegenerates ? THRESH_POINTS_ARE_SAME : 0.0f;
+		TMultiMap<int32, int32> OverlappingCorners;
+		FindOverlappingCorners(OverlappingCorners, OutRawMesh, ComparisonThreshold);
+
+		RecomputeTangentsAndNormalsForRawMesh( bRecomputeTangents, bRecomputeNormals, InBuildSettings, OverlappingCorners, OutRawMesh );
+	}
+}
+
+void FMeshUtilities::RecomputeTangentsAndNormalsForRawMesh(bool bRecomputeTangents, bool bRecomputeNormals, const FMeshBuildSettings& InBuildSettings, const TMultiMap<int32, int32>& InOverlappingCorners, FRawMesh &OutRawMesh) const
 {
 	const int32 NumWedges = OutRawMesh.WedgeIndices.Num();
 
@@ -5173,10 +5186,6 @@ void FMeshUtilities::RecomputeTangentsAndNormalsForRawMesh(bool bRecomputeTangen
 	// Compute any missing tangents.
 	if (bRecomputeNormals || bRecomputeTangents)
 	{
-		float ComparisonThreshold = InBuildSettings.bRemoveDegenerates ? THRESH_POINTS_ARE_SAME : 0.0f;
-		TMultiMap<int32, int32> OverlappingCorners;
-		FindOverlappingCorners(OverlappingCorners, OutRawMesh, ComparisonThreshold);
-
 		// Static meshes always blend normals of overlapping corners.
 		uint32 TangentOptions = ETangentOptions::BlendOverlappingNormals;
 		if (InBuildSettings.bRemoveDegenerates)
@@ -5184,13 +5193,14 @@ void FMeshUtilities::RecomputeTangentsAndNormalsForRawMesh(bool bRecomputeTangen
 			// If removing degenerate triangles, ignore them when computing tangents.
 			TangentOptions |= ETangentOptions::IgnoreDegenerateTriangles;
 		}
+
 		if (InBuildSettings.bUseMikkTSpace)
 		{
-			ComputeTangents_MikkTSpace(OutRawMesh, OverlappingCorners, TangentOptions);
+			ComputeTangents_MikkTSpace(OutRawMesh, InOverlappingCorners, TangentOptions);
 		}
 		else
 		{
-			ComputeTangents(OutRawMesh, OverlappingCorners, TangentOptions);
+			ComputeTangents(OutRawMesh, InOverlappingCorners, TangentOptions);
 		}
 	}
 
diff --git a/Engine/Source/Developer/MeshUtilities/Private/MeshUtilitiesPrivate.h b/Engine/Source/Developer/MeshUtilities/Private/MeshUtilitiesPrivate.h
--- a/Engine/Source/Developer/MeshUtilities/Private/MeshUtilitiesPrivate.h
+++ b/Engine/Source/Developer/MeshUtilities/Private/MeshUtilitiesPrivate.h
@@ -107,7 +107,7 @@ private:
 		FDistanceFieldVolumeData& OutData) override;
 
 	virtual void RecomputeTangentsAndNormalsForRawMesh(bool bRecomputeTangents, bool bRecomputeNormals, const FMeshBuildSettings& InBuildSettings, FRawMesh &OutRawMesh) const override;
-
+	virtual void RecomputeTangentsAndNormalsForRawMesh(bool bRecomputeTangents, bool bRecomputeNormals, const FMeshBuildSettings& InBuildSettings, const TMultiMap<int32, int32>& InOverlappingCorners, FRawMesh &OutRawMesh) const override;
 
 	virtual bool GenerateUniqueUVsForStaticMesh(const FRawMesh& RawMesh, int32 TextureResolution, TArray<FVector2D>& OutTexCoords) const override;
 
diff --git a/Engine/Source/Developer/MeshUtilities/Public/LayoutUV.h b/Engine/Source/Developer/MeshUtilities/Public/LayoutUV.h
--- a/Engine/Source/Developer/MeshUtilities/Public/LayoutUV.h
+++ b/Engine/Source/Developer/MeshUtilities/Public/LayoutUV.h
@@ -48,7 +48,7 @@ class MESHUTILITIES_API FLayoutUV
 public:
 				FLayoutUV( FRawMesh* InMesh, uint32 InSrcChannel, uint32 InDstChannel, uint32 InTextureResolution );
 
-	void		FindCharts( const TMultiMap<int32,int32>& OverlappingCorners );
+	int32		FindCharts( const TMultiMap<int32,int32>& OverlappingCorners );
 	bool		FindBestPacking();
 	void		CommitPackedUVs();
 
diff --git a/Engine/Source/Developer/MeshUtilities/Public/MeshUtilities.h b/Engine/Source/Developer/MeshUtilities/Public/MeshUtilities.h
--- a/Engine/Source/Developer/MeshUtilities/Public/MeshUtilities.h
+++ b/Engine/Source/Developer/MeshUtilities/Public/MeshUtilities.h
@@ -336,6 +336,7 @@ public:
 	virtual void CalculateOverlappingCorners(const TArray<FVector>& InVertices, const TArray<uint32>& InIndices, bool bIgnoreDegenerateTriangles, TMultiMap<int32, int32>& OutOverlappingCorners) const = 0;
 
 	virtual void RecomputeTangentsAndNormalsForRawMesh(bool bRecomputeTangents, bool bRecomputeNormals, const FMeshBuildSettings& InBuildSettings, FRawMesh &OutRawMesh) const = 0;
+	virtual void RecomputeTangentsAndNormalsForRawMesh(bool bRecomputeTangents, bool bRecomputeNormals, const FMeshBuildSettings& InBuildSettings, const TMultiMap<int32, int32>& InOverlappingCorners, FRawMesh &OutRawMesh) const = 0;
 
 	virtual void FindOverlappingCorners(TMultiMap<int32, int32>& OutOverlappingCorners, const TArray<FVector>& InVertices, const TArray<uint32>& InIndices, float ComparisonThreshold) const = 0;
 };
diff --git a/Engine/Source/Developer/OutputLog/Private/OutputLogModule.cpp b/Engine/Source/Developer/OutputLog/Private/OutputLogModule.cpp
--- a/Engine/Source/Developer/OutputLog/Private/OutputLogModule.cpp
+++ b/Engine/Source/Developer/OutputLog/Private/OutputLogModule.cpp
@@ -1,6 +1,7 @@
 // Copyright 1998-2017 Epic Games, Inc. All Rights Reserved.
 
 #include "OutputLogModule.h"
+#include "Features/IModularFeatures.h"
 #include "Modules/ModuleManager.h"
 #include "Framework/Application/SlateApplication.h"
 #include "Textures/SlateIcon.h"
@@ -88,6 +89,9 @@ TSharedRef<SDockTab> SpawnDeviceOutputLog( const FSpawnTabArgs& Args )
 
 void FOutputLogModule::StartupModule()
 {
+	CmdExec = MakeShared<FConsoleCommandExecutor>();
+	IModularFeatures::Get().RegisterModularFeature(IConsoleCommandExecutor::ModularFeatureName(), CmdExec.Get());
+
 	FGlobalTabmanager::Get()->RegisterNomadTabSpawner(OutputLogModule::OutputLogTabName, FOnSpawnTab::CreateStatic( &SpawnOutputLog ) )
 		.SetDisplayName(NSLOCTEXT("UnrealEditor", "OutputLogTab", "Output Log"))
 		.SetTooltipText(NSLOCTEXT("UnrealEditor", "OutputLogTooltipText", "Open the Output Log tab."))
@@ -106,6 +110,8 @@ void FOutputLogModule::StartupModule()
 
 void FOutputLogModule::ShutdownModule()
 {
+	IModularFeatures::Get().UnregisterModularFeature(IConsoleCommandExecutor::ModularFeatureName(), CmdExec.Get());
+
 	if (FSlateApplication::IsInitialized())
 	{
 		FGlobalTabmanager::Get()->UnregisterNomadTabSpawner(OutputLogModule::OutputLogTabName);
@@ -113,14 +119,13 @@ void FOutputLogModule::ShutdownModule()
 	}
 }
 
-TSharedRef< SWidget > FOutputLogModule::MakeConsoleInputBox( TSharedPtr< SEditableTextBox >& OutExposedEditableTextBox ) const
+TSharedRef< SWidget > FOutputLogModule::MakeConsoleInputBox( TSharedPtr< SMultiLineEditableTextBox >& OutExposedEditableTextBox ) const
 {
 	TSharedRef< SConsoleInputBox > NewConsoleInputBox = SNew( SConsoleInputBox );
 	OutExposedEditableTextBox = NewConsoleInputBox->GetEditableTextBox();
 	return NewConsoleInputBox;
 }
 
-
 void FOutputLogModule::ToggleDebugConsoleForWindow( const TSharedRef< SWindow >& Window, const EDebugConsoleStyle::Type InStyle, const FDebugConsoleDelegates& DebugConsoleDelegates )
 {
 	bool bShouldOpen = true;
@@ -171,7 +176,6 @@ void FOutputLogModule::ToggleDebugConsoleForWindow( const TSharedRef< SWindow >&
 	}
 }
 
-
 void FOutputLogModule::CloseDebugConsole()
 {
 	TSharedPtr< SWidget > PinnedDebugConsole( DebugConsole.Pin() );
diff --git a/Engine/Source/Developer/OutputLog/Private/SDebugConsole.cpp b/Engine/Source/Developer/OutputLog/Private/SDebugConsole.cpp
--- a/Engine/Source/Developer/OutputLog/Private/SDebugConsole.cpp
+++ b/Engine/Source/Developer/OutputLog/Private/SDebugConsole.cpp
@@ -32,24 +32,24 @@ void SDebugConsole::Construct( const FArguments& InArgs, const EDebugConsoleStyl
 		.AutoHeight()
 		[
 			SNew( SVerticalBox )
-				.Visibility( this, &SDebugConsole::MakeVisibleIfLogIsShown )
+			.Visibility( this, &SDebugConsole::MakeVisibleIfLogIsShown )
 					
-				+SVerticalBox::Slot()
-				.AutoHeight()
-				.Padding( 10.0f )
+			+SVerticalBox::Slot()
+			.AutoHeight()
+			.Padding( 10.0f )
+			[
+				SNew(SBox)
+				.HeightOverride( 200.0f )
 				[
-					SNew(SBox)
-					.HeightOverride( 200.0f )
+					SNew( SBorder )
+					.BorderImage( FEditorStyle::GetBrush( "ToolPanel.GroupBorder" ) )
+					.ColorAndOpacity( this, &SDebugConsole::GetAnimatedColorAndOpacity )
+					.BorderBackgroundColor( this, &SDebugConsole::GetAnimatedSlateColor )
 					[
-						SNew( SBorder )
-							.BorderImage( FEditorStyle::GetBrush( "ToolPanel.GroupBorder" ) )
-							.ColorAndOpacity( this, &SDebugConsole::GetAnimatedColorAndOpacity )
-							.BorderBackgroundColor( this, &SDebugConsole::GetAnimatedSlateColor )
-							[
-								SNew( SSpacer )
-							]
+						SNew( SSpacer )
 					]
 				]
+			]
 		]
 
 		+SVerticalBox::Slot()
@@ -57,33 +57,19 @@ void SDebugConsole::Construct( const FArguments& InArgs, const EDebugConsoleStyl
 		.Padding( 10.0f )
 		[
 			SNew(SBox)
-			.HeightOverride( 26.0f )
-			.HAlign( HAlign_Left )
+			.WidthOverride( 400.0f )
+			.MaxDesiredHeight( 100.0f )
 			[
 				SNew( SBorder )
-				.Padding( FMargin(2) )
+				.VAlign( VAlign_Center )
+				.Padding( FMargin(5.0f, 2.0f) )
 				.BorderImage( FEditorStyle::GetBrush( "DebugConsole.Background" ) )
 				.ColorAndOpacity( this, &SDebugConsole::GetAnimatedColorAndOpacity )
 				.BorderBackgroundColor( this, &SDebugConsole::GetAnimatedSlateColor )
 				[
-					SNew(SHorizontalBox)
-					+ SHorizontalBox::Slot()
-					.AutoWidth()
-					.Padding(3.0f)
-					.VAlign(VAlign_Center)
-					[
-						SNew(STextBlock)
-						.Text(NSLOCTEXT("Console", "ConsoleLabel", "Console"))
-
-					]
-					+ SHorizontalBox::Slot()
-					.Padding(5.0f, 2.0f)
-					.VAlign(VAlign_Center)
-					.MaxWidth(400.0f)
-					[
-						SAssignNew(ConsoleInputBox, SConsoleInputBox)
-						.OnConsoleCommandExecuted(DebugConsoleDelegates->OnConsoleCommandExecuted)
-					]
+					SAssignNew(ConsoleInputBox, SConsoleInputBox)
+					.OnConsoleCommandExecuted(DebugConsoleDelegates->OnConsoleCommandExecuted)
+					.OnCloseConsole(DebugConsoleDelegates->OnCloseConsole)
 				]
 			]
 		]
diff --git a/Engine/Source/Developer/OutputLog/Private/SDebugConsole.h b/Engine/Source/Developer/OutputLog/Private/SDebugConsole.h
--- a/Engine/Source/Developer/OutputLog/Private/SDebugConsole.h
+++ b/Engine/Source/Developer/OutputLog/Private/SDebugConsole.h
@@ -11,7 +11,7 @@
 #include "Widgets/SCompoundWidget.h"
 #include "OutputLogModule.h"
 
-class SEditableTextBox;
+class SMultiLineEditableTextBox;
 
 /**
  * Debug console widget, designed to be summoned on top of a viewport or window
@@ -50,7 +50,7 @@ protected:
 private:
 
 	/** Editable text box for this debug console's input line */
-	TSharedPtr< SEditableTextBox > EditableTextBox;
+	TSharedPtr< SMultiLineEditableTextBox > EditableTextBox;
 
 	/** Current style of the debug console.  Can be changed on the fly. */
 	EDebugConsoleStyle::Type CurrentStyle;
diff --git a/Engine/Source/Developer/OutputLog/Private/SOutputLog.cpp b/Engine/Source/Developer/OutputLog/Private/SOutputLog.cpp
--- a/Engine/Source/Developer/OutputLog/Private/SOutputLog.cpp
+++ b/Engine/Source/Developer/OutputLog/Private/SOutputLog.cpp
@@ -4,7 +4,7 @@
 #include "Framework/Text/TextRange.h"
 #include "Framework/Text/IRun.h"
 #include "Framework/Text/TextLayout.h"
-#include "HAL/IConsoleManager.h"
+#include "Misc/ConfigCacheIni.h"
 #include "Misc/OutputDeviceHelper.h"
 #include "SlateOptMacros.h"
 #include "Textures/SlateIcon.h"
@@ -21,148 +21,167 @@
 #include "EditorStyleSettings.h"
 #include "EngineGlobals.h"
 #include "Editor.h"
+#include "Toolkits/GlobalEditorCommonCommands.h"
 #include "Engine/LocalPlayer.h"
 #include "GameFramework/GameStateBase.h"
 #include "Widgets/Input/SSearchBox.h"
+#include "Features/IModularFeatures.h"
 
 #define LOCTEXT_NAMESPACE "SOutputLog"
 
-/** Expression context to test the given messages against the current text filter */
-class FLogFilter_TextFilterExpressionContext : public ITextFilterExpressionContext
-{
-public:
-	explicit FLogFilter_TextFilterExpressionContext(const FLogMessage& InMessage) : Message(&InMessage) {}
+ FName FConsoleCommandExecutor::StaticName()
+ {
+	 static const FName CmdExecName = TEXT("Cmd");
+	 return CmdExecName;
+ }
 
-	/** Test the given value against the strings extracted from the current item */
-	virtual bool TestBasicStringExpression(const FTextFilterString& InValue, const ETextFilterTextComparisonMode InTextComparisonMode) const override { return TextFilterUtils::TestBasicStringExpression(*Message->Message, InValue, InTextComparisonMode); }
+ FName FConsoleCommandExecutor::GetName() const
+ {
+	 return StaticName();
+ }
 
-	/**
-	* Perform a complex expression test for the current item
-	* No complex expressions in this case - always returns false
-	*/
-	virtual bool TestComplexExpression(const FName& InKey, const FTextFilterString& InValue, const ETextFilterComparisonOperation InComparisonOperation, const ETextFilterTextComparisonMode InTextComparisonMode) const override { return false; }
+FText FConsoleCommandExecutor::GetDisplayName() const
+{
+	return LOCTEXT("ConsoleCommandExecutorDisplayName", "Cmd");
+}
 
-private:
-	/** Message that is being filtered */
-	const FLogMessage* Message;
-};
+FText FConsoleCommandExecutor::GetDescription() const
+{
+	return LOCTEXT("ConsoleCommandExecutorDescription", "Execute Unreal Console Commands");
+}
 
-/** Custom console editable text box whose only purpose is to prevent some keys from being typed */
-class SConsoleEditableTextBox : public SEditableTextBox
+FText FConsoleCommandExecutor::GetHintText() const
 {
-public:
-	SLATE_BEGIN_ARGS( SConsoleEditableTextBox ) {}
-		
-		/** Hint text that appears when there is no text in the text box */
-		SLATE_ATTRIBUTE(FText, HintText)
+	return LOCTEXT("ConsoleCommandExecutorHintText", "Enter Console Command");
+}
 
-		/** Called whenever the text is changed interactively by the user */
-		SLATE_EVENT(FOnTextChanged, OnTextChanged)
+void FConsoleCommandExecutor::GetAutoCompleteSuggestions(const TCHAR* Input, TArray<FString>& Out)
+{
+	auto OnConsoleVariable = [&Out](const TCHAR *Name, IConsoleObject* CVar)
+	{
+#if (UE_BUILD_SHIPPING || UE_BUILD_TEST)
+		if (CVar->TestFlags(ECVF_Cheat))
+		{
+			return;
+		}
+#endif // (UE_BUILD_SHIPPING || UE_BUILD_TEST)
+		if (CVar->TestFlags(ECVF_Unregistered))
+		{
+			return;
+		}
 
-		/** Called whenever the text is committed.  This happens when the user presses enter or the text box loses focus. */
-		SLATE_EVENT(FOnTextCommitted, OnTextCommitted)
+		Out.Add(Name);
+	};
 
-	SLATE_END_ARGS()
+	IConsoleManager::Get().ForEachConsoleObjectThatContains(FConsoleObjectVisitor::CreateLambda(OnConsoleVariable), Input);
+}
 
+void FConsoleCommandExecutor::GetExecHistory(TArray<FString>& Out)
+{
+	IConsoleManager::Get().GetConsoleHistory(TEXT(""), Out);
+}
 
-	void Construct( const FArguments& InArgs )
-	{
-		SetStyle(&FCoreStyle::Get().GetWidgetStyle< FEditableTextBoxStyle >("NormalEditableTextBox"));
+bool FConsoleCommandExecutor::Exec(const TCHAR* Input)
+{
+	IConsoleManager::Get().AddConsoleHistoryEntry(TEXT(""), Input);
 
-		SBorder::Construct(SBorder::FArguments()
-			.BorderImage(this, &SConsoleEditableTextBox::GetConsoleBorder)
-			.BorderBackgroundColor(Style->BackgroundColor)
-			.ForegroundColor(Style->ForegroundColor)
-			.Padding(Style->Padding)
-			[
-				SAssignNew( EditableText, SConsoleEditableText )
-				.HintText( InArgs._HintText )
-				.OnTextChanged( InArgs._OnTextChanged )
-				.OnTextCommitted( InArgs._OnTextCommitted )
-			] );
+	bool bWasHandled = false;
+	UWorld* World = nullptr;
+	UWorld* OldWorld = nullptr;
+
+	// The play world needs to handle these commands if it exists
+	if (GIsEditor && GEditor->PlayWorld && !GIsPlayInEditorWorld)
+	{
+		World = GEditor->PlayWorld;
+		OldWorld = SetPlayInEditorWorld(GEditor->PlayWorld);
 	}
 
-private:
-	class SConsoleEditableText : public SEditableText
+	ULocalPlayer* Player = GEngine->GetDebugLocalPlayer();
+	if (Player)
 	{
-	public:
-		SLATE_BEGIN_ARGS( SConsoleEditableText ) {}
-			/** The text that appears when there is nothing typed into the search box */
-			SLATE_ATTRIBUTE(FText, HintText)
-			/** Called whenever the text is changed interactively by the user */
-			SLATE_EVENT(FOnTextChanged, OnTextChanged)
-
-			/** Called whenever the text is committed.  This happens when the user presses enter or the text box loses focus. */
-			SLATE_EVENT(FOnTextCommitted, OnTextCommitted)
-		SLATE_END_ARGS()
-
-		void Construct( const FArguments& InArgs )
+		UWorld* PlayerWorld = Player->GetWorld();
+		if (!World)
 		{
-			SEditableText::Construct
-			( 
-				SEditableText::FArguments()
-				.HintText( InArgs._HintText )
-				.OnTextChanged( InArgs._OnTextChanged )
-				.OnTextCommitted( InArgs._OnTextCommitted ) 
-				.ClearKeyboardFocusOnCommit( false )
-				.IsCaretMovedWhenGainFocus( false ) 
-				.MinDesiredWidth( 400.0f )
-			);
+			World = PlayerWorld;
 		}
+		bWasHandled = Player->Exec(PlayerWorld, Input, *GLog);
+	}
 
-		virtual FReply OnKeyDown( const FGeometry& MyGeometry, const FKeyEvent& InKeyEvent )
+	if (!World)
+	{
+		World = GEditor->GetEditorWorldContext().World();
+	}
+	if (World)
+	{
+		if (!bWasHandled)
 		{
-			// Special case handling.  Intercept the tilde key.  It is not suitable for typing in the console
-			if( InKeyEvent.GetKey() == EKeys::Tilde )
+			AGameModeBase* const GameMode = World->GetAuthGameMode();
+			AGameStateBase* const GameState = World->GetGameState();
+			if (GameMode && GameMode->ProcessConsoleExec(Input, *GLog, nullptr))
 			{
-				return FReply::Unhandled();
+				bWasHandled = true;
 			}
-			else
+			else if (GameState && GameState->ProcessConsoleExec(Input, *GLog, nullptr))
 			{
-				return SEditableText::OnKeyDown( MyGeometry, InKeyEvent );
+				bWasHandled = true;
 			}
 		}
 
-		virtual FReply OnKeyChar( const FGeometry& MyGeometry, const FCharacterEvent& InCharacterEvent )
+		if (!bWasHandled && !Player)
 		{
-			// Special case handling.  Intercept the tilde key.  It is not suitable for typing in the console
-			if( InCharacterEvent.GetCharacter() != 0x60 )
+			if (GIsEditor)
 			{
-				return SEditableText::OnKeyChar( MyGeometry, InCharacterEvent );
+				bWasHandled = GEditor->Exec(World, Input, *GLog);
 			}
 			else
 			{
-				return FReply::Unhandled();
+				bWasHandled = GEngine->Exec(World, Input, *GLog);
 			}
 		}
+	}
 
-	};
-
-	/** @return Border image for the text box based on the hovered and focused state */
-	const FSlateBrush* GetConsoleBorder() const
+	// Restore the old world of there was one
+	if (OldWorld)
 	{
-		if (EditableText->HasKeyboardFocus())
-		{
-			return &Style->BackgroundImageFocused;
-		}
-		else
-		{
-			if (EditableText->IsHovered())
-			{
-				return &Style->BackgroundImageHovered;
-			}
-			else
-			{
-				return &Style->BackgroundImageNormal;
-			}
-		}
+		RestoreEditorWorld(OldWorld);
 	}
 
+	return bWasHandled;
+}
+
+bool FConsoleCommandExecutor::AllowHotKeyClose() const
+{
+	return true;
+}
+
+bool FConsoleCommandExecutor::AllowMultiLine() const
+{
+	return false;
+}
+
+/** Expression context to test the given messages against the current text filter */
+class FLogFilter_TextFilterExpressionContext : public ITextFilterExpressionContext
+{
+public:
+	explicit FLogFilter_TextFilterExpressionContext(const FLogMessage& InMessage) : Message(&InMessage) {}
+
+	/** Test the given value against the strings extracted from the current item */
+	virtual bool TestBasicStringExpression(const FTextFilterString& InValue, const ETextFilterTextComparisonMode InTextComparisonMode) const override { return TextFilterUtils::TestBasicStringExpression(*Message->Message, InValue, InTextComparisonMode); }
+
+	/**
+	* Perform a complex expression test for the current item
+	* No complex expressions in this case - always returns false
+	*/
+	virtual bool TestComplexExpression(const FName& InKey, const FTextFilterString& InValue, const ETextFilterComparisonOperation InComparisonOperation, const ETextFilterTextComparisonMode InTextComparisonMode) const override { return false; }
+
+private:
+	/** Message that is being filtered */
+	const FLogMessage* Message;
 };
 
 SConsoleInputBox::SConsoleInputBox()
-	: SelectedSuggestion(-1)
-	, bIgnoreUIUpdate(false)
+	: bIgnoreUIUpdate(false)
+	, bHasTicked(false)
 {
 }
 
@@ -171,41 +190,87 @@ void SConsoleInputBox::Construct( const FArguments& InArgs )
 {
 	OnConsoleCommandExecuted = InArgs._OnConsoleCommandExecuted;
 	ConsoleCommandCustomExec = InArgs._ConsoleCommandCustomExec;
+	OnCloseConsole = InArgs._OnCloseConsole;
+
+	PreferredCommandExecutorName = FConsoleCommandExecutor::StaticName();
+	if (!ConsoleCommandCustomExec.IsBound()) // custom execs always show the default executor in the UI (which has the selector disabled)
+	{
+		FString PreferredCommandExecutorStr;
+		if (GConfig->GetString(TEXT("OutputLog"), TEXT("PreferredCommandExecutor"), PreferredCommandExecutorStr, GEditorPerProjectIni))
+		{
+			PreferredCommandExecutorName = *PreferredCommandExecutorStr;
+		}
+	}
+	SyncActiveCommandExecutor();
+
+	IModularFeatures::Get().OnModularFeatureRegistered().AddSP(this, &SConsoleInputBox::OnCommandExecutorRegistered);
+	IModularFeatures::Get().OnModularFeatureUnregistered().AddSP(this, &SConsoleInputBox::OnCommandExecutorUnregistered);
 
 	ChildSlot
 	[
 		SAssignNew( SuggestionBox, SMenuAnchor )
-			.Placement( InArgs._SuggestionListPlacement )
+		.Placement( InArgs._SuggestionListPlacement )
+		[
+			SNew(SHorizontalBox)
+
+			+SHorizontalBox::Slot()
+			.AutoWidth()
+			.Padding(FMargin(0.0f, 0.0f, 4.0f, 0.0f))
+			[
+				SNew(SComboButton)
+				.IsEnabled(this, &SConsoleInputBox::IsCommandExecutorMenuEnabled)
+				.ComboButtonStyle(FEditorStyle::Get(), "GenericFilters.ComboButtonStyle")
+				.ForegroundColor(FLinearColor::White)
+				.ContentPadding(0)
+				.OnGetMenuContent(this, &SConsoleInputBox::GetCommandExecutorMenuContent)
+				.ButtonContent()
+				[
+					SNew(STextBlock)
+					.Text(this, &SConsoleInputBox::GetActiveCommandExecutorDisplayName)
+				]
+			]
+
+			+SHorizontalBox::Slot()
 			[
-				SAssignNew(InputText, SConsoleEditableTextBox)
-					.OnTextCommitted(this, &SConsoleInputBox::OnTextCommitted)
-					.HintText( NSLOCTEXT( "ConsoleInputBox", "TypeInConsoleHint", "Enter console command" ) )
-					.OnTextChanged(this, &SConsoleInputBox::OnTextChanged)
+				SAssignNew(InputText, SMultiLineEditableTextBox)
+				.Font(FEditorStyle::Get().GetWidgetStyle<FTextBlockStyle>("Log.Normal").Font)
+				.HintText(this, &SConsoleInputBox::GetActiveCommandExecutorHintText)
+				.AllowMultiLine(this, &SConsoleInputBox::GetActiveCommandExecutorAllowMultiLine)
+				.OnTextCommitted(this, &SConsoleInputBox::OnTextCommitted)
+				.OnTextChanged(this, &SConsoleInputBox::OnTextChanged)
+				.OnKeyCharHandler(this, &SConsoleInputBox::OnKeyCharHandler)
+				.OnKeyDownHandler(this, &SConsoleInputBox::OnKeyDownHandler)
+				.OnIsTypedCharValid(FOnIsTypedCharValid::CreateLambda([](const TCHAR InCh) { return true; })) // allow tabs to be typed into the field
+				.ClearKeyboardFocusOnCommit(false)
+				.ModiferKeyForNewLine(EModifierKey::Shift)
 			]
-			.MenuContent
-			(
-				SNew(SBorder)
-				.BorderImage(FEditorStyle::GetBrush("Menu.Background"))
-				.Padding( FMargin(2) )
+		]
+		.MenuContent
+		(
+			SNew(SBorder)
+			.BorderImage(FEditorStyle::GetBrush("Menu.Background"))
+			.Padding( FMargin(2) )
+			[
+				SNew(SBox)
+				.HeightOverride(250) // avoids flickering, ideally this would be adaptive to the content without flickering
 				[
-					SNew(SBox)
-					.HeightOverride(250) // avoids flickering, ideally this would be adaptive to the content without flickering
-					[
-						SAssignNew(SuggestionListView, SListView< TSharedPtr<FString> >)
-							.ListItemsSource(&Suggestions)
-							.SelectionMode( ESelectionMode::Single )							// Ideally the mouse over would not highlight while keyboard controls the UI
-							.OnGenerateRow(this, &SConsoleInputBox::MakeSuggestionListItemWidget)
-							.OnSelectionChanged(this, &SConsoleInputBox::SuggestionSelectionChanged)
-							.ItemHeight(18)
-					]
+					SAssignNew(SuggestionListView, SListView< TSharedPtr<FString> >)
+					.ListItemsSource(&Suggestions.SuggestionsList)
+					.SelectionMode( ESelectionMode::Single )							// Ideally the mouse over would not highlight while keyboard controls the UI
+					.OnGenerateRow(this, &SConsoleInputBox::MakeSuggestionListItemWidget)
+					.OnSelectionChanged(this, &SConsoleInputBox::SuggestionSelectionChanged)
+					.ItemHeight(18)
 				]
-			)
+			]
+		)
 	];
 }
 END_SLATE_FUNCTION_BUILD_OPTIMIZATION
 
 void SConsoleInputBox::Tick( const FGeometry& AllottedGeometry, const double InCurrentTime, const float InDeltaTime )
 {
+	bHasTicked = true;
+
 	if (!GIntraFrameDebuggingGameThread && !IsEnabled())
 	{
 		SetEnabled(true);
@@ -224,27 +289,24 @@ void SConsoleInputBox::SuggestionSelectionChanged(TSharedPtr<FString> NewValue,
 		return;
 	}
 
-	for(int32 i = 0; i < Suggestions.Num(); ++i)
+	Suggestions.SelectedSuggestion = Suggestions.SuggestionsList.IndexOfByPredicate([&NewValue](const TSharedPtr<FString>& InSuggestion)
 	{
-		if(NewValue == Suggestions[i])
-		{
-			SelectedSuggestion = i;
-			MarkActiveSuggestion();
+		return InSuggestion == NewValue;
+	});
 
-			// If the user selected this suggestion by clicking on it, then go ahead and close the suggestion
-			// box as they've chosen the suggestion they're interested in.
-			if( SelectInfo == ESelectInfo::OnMouseClick )
-			{
-				SuggestionBox->SetIsOpen( false );
-			}
+	MarkActiveSuggestion();
 
-			// Ideally this would set the focus back to the edit control
-//			FWidgetPath WidgetToFocusPath;
-//			FSlateApplication::Get().GeneratePathToWidgetUnchecked( InputText.ToSharedRef(), WidgetToFocusPath );
-//			FSlateApplication::Get().SetKeyboardFocus( WidgetToFocusPath, EFocusCause::SetDirectly );
-			break;
-		}
+	// If the user selected this suggestion by clicking on it, then go ahead and close the suggestion
+	// box as they've chosen the suggestion they're interested in.
+	if( SelectInfo == ESelectInfo::OnMouseClick )
+	{
+		SuggestionBox->SetIsOpen( false );
 	}
+
+	// Ideally this would set the focus back to the edit control
+//	FWidgetPath WidgetToFocusPath;
+//	FSlateApplication::Get().GeneratePathToWidgetUnchecked( InputText.ToSharedRef(), WidgetToFocusPath );
+//	FSlateApplication::Get().SetKeyboardFocus( WidgetToFocusPath, EFocusCause::SetDirectly );
 }
 
 
@@ -252,25 +314,10 @@ TSharedRef<ITableRow> SConsoleInputBox::MakeSuggestionListItemWidget(TSharedPtr<
 {
 	check(Text.IsValid());
 
-	FString Left, Mid, Right, TempRight, Combined;
-
-	if(Text->Split(TEXT("\t"), &Left, &TempRight))
-	{
-		if (TempRight.Split(TEXT("\t"), &Mid, &Right))
-		{
-			Combined = Left + Mid + Right;
-		}
-		else
-		{
-			Combined = Left + Right;
-		}
-	}
-	else
-	{
-		Combined = *Text;
-	}
-
-	FText HighlightText = FText::FromString(Mid);
+	FString SanitizedText = *Text;
+	SanitizedText.ReplaceInline(TEXT("\r\n"), TEXT("\n"), ESearchCase::CaseSensitive);
+	SanitizedText.ReplaceInline(TEXT("\r"), TEXT(" "), ESearchCase::CaseSensitive);
+	SanitizedText.ReplaceInline(TEXT("\n"), TEXT(" "), ESearchCase::CaseSensitive);
 
 	return
 		SNew(STableRow< TSharedPtr<FString> >, OwnerTable)
@@ -279,35 +326,13 @@ TSharedRef<ITableRow> SConsoleInputBox::MakeSuggestionListItemWidget(TSharedPtr<
 			.WidthOverride(300)			// to enforce some minimum width, ideally we define the minimum, not a fixed width
 			[
 				SNew(STextBlock)
-				.Text(FText::FromString(Combined))
-				.TextStyle( FEditorStyle::Get(), "Log.Normal")
-				.HighlightText(HighlightText)
+				.Text(FText::FromString(SanitizedText))
+				.TextStyle(FEditorStyle::Get(), "Log.Normal")
+				.HighlightText(Suggestions.SuggestionsHighlight)
 			]
 		];
 }
 
-class FConsoleVariableAutoCompleteVisitor 
-{
-public:
-	// @param Name must not be 0
-	// @param CVar must not be 0
-	static void OnConsoleVariable(const TCHAR *Name, IConsoleObject* CVar,TArray<FString>& Sink)
-	{
-#if (UE_BUILD_SHIPPING || UE_BUILD_TEST)
-		if(CVar->TestFlags(ECVF_Cheat))
-		{
-			return;
-		}
-#endif // (UE_BUILD_SHIPPING || UE_BUILD_TEST)
-		if(CVar->TestFlags(ECVF_Unregistered))
-		{
-			return;
-		}
-
-		Sink.Add(Name);
-	}
-};
-
 void SConsoleInputBox::OnTextChanged(const FText& InText)
 {
 	if(bIgnoreUIUpdate)
@@ -319,29 +344,14 @@ void SConsoleInputBox::OnTextChanged(const FText& InText)
 	if(!InputTextStr.IsEmpty())
 	{
 		TArray<FString> AutoCompleteList;
-
-		// console variables
+		
+		if (ActiveCommandExecutor)
 		{
-			IConsoleManager::Get().ForEachConsoleObjectThatContains(
-				FConsoleObjectVisitor::CreateStatic< TArray<FString>& >(
-				&FConsoleVariableAutoCompleteVisitor::OnConsoleVariable,
-				AutoCompleteList ), *InputTextStr);
+			ActiveCommandExecutor->GetAutoCompleteSuggestions(*InputTextStr, AutoCompleteList);
 		}
 
 		AutoCompleteList.Sort();
-
-		for(uint32 i = 0; i < (uint32)AutoCompleteList.Num(); ++i)
-		{
-			FString &ref = AutoCompleteList[i];
-			int32 Start = ref.Find(InputTextStr);
-
-			if (Start != INDEX_NONE)
-			{
-				ref = ref.Left(Start) + TEXT("\t") + ref.Mid(Start, InputTextStr.Len()) + TEXT("\t") + ref.RightChop(Start + InputTextStr.Len());
-			}
-		}
-
-		SetSuggestions(AutoCompleteList, false);
+		SetSuggestions(AutoCompleteList, FText::FromString(InputTextStr));
 	}
 	else
 	{
@@ -355,8 +365,6 @@ void SConsoleInputBox::OnTextCommitted( const FText& InText, ETextCommit::Type C
 	{
 		if (!InText.IsEmpty())
 		{
-			IConsoleManager::Get().AddConsoleHistoryEntry( *InText.ToString() );
-
 			// Copy the exec text string out so we can clear the widget's contents.  If the exec command spawns
 			// a new window it can cause the text box to lose focus, which will result in this function being
 			// re-entered.  We want to make sure the text string is empty on re-entry, so we'll clear it out
@@ -370,69 +378,12 @@ void SConsoleInputBox::OnTextCommitted( const FText& InText, ETextCommit::Type C
 			// Exec!
 			if (ConsoleCommandCustomExec.IsBound())
 			{
+				IConsoleManager::Get().AddConsoleHistoryEntry(TEXT(""), *ExecString);
 				ConsoleCommandCustomExec.Execute(ExecString);
 			}
-			else
+			else if (ActiveCommandExecutor)
 			{
-				bool bWasHandled = false;
-				UWorld* World = NULL;
-				UWorld* OldWorld = NULL;
-
-				// The play world needs to handle these commands if it exists
-				if( GIsEditor && GEditor->PlayWorld && !GIsPlayInEditorWorld )
-				{
-					World = GEditor->PlayWorld;
-					OldWorld = SetPlayInEditorWorld( GEditor->PlayWorld );
-				}
-				
-				ULocalPlayer* Player = GEngine->GetDebugLocalPlayer();
-				if( Player )
-				{
-					UWorld* PlayerWorld = Player->GetWorld();
-					if( !World )
-					{
-						World = PlayerWorld;
-					}
-					bWasHandled = Player->Exec( PlayerWorld, *ExecString, *GLog );
-				}
-					
-				if( !World )
-				{
-					World = GEditor->GetEditorWorldContext().World();
-				}
-				if( World )
-				{
-					if( !bWasHandled )
-					{					
-						AGameModeBase* const GameMode = World->GetAuthGameMode();
-						AGameStateBase* const GameState = World->GetGameState();
-						if( GameMode && GameMode->ProcessConsoleExec( *ExecString, *GLog, NULL ) )
-						{
-							bWasHandled = true;
-						}
-						else if (GameState && GameState->ProcessConsoleExec(*ExecString, *GLog, NULL))
-						{
-							bWasHandled = true;
-						}
-					}
-
-					if( !bWasHandled && !Player)
-					{
-						if( GIsEditor )
-						{
-							bWasHandled = GEditor->Exec( World, *ExecString, *GLog );
-						}
-						else
-						{
-							bWasHandled = GEngine->Exec( World, *ExecString, *GLog );
-						}	
-					}
-				}
-				// Restore the old world of there was one
-				if( OldWorld )
-				{
-					RestoreEditorWorld( OldWorld );
-				}
+				ActiveCommandExecutor->Exec(*ExecString);
 			}
 		}
 
@@ -448,71 +399,79 @@ FReply SConsoleInputBox::OnPreviewKeyDown(const FGeometry& MyGeometry, const FKe
 	{
 		if(KeyEvent.GetKey() == EKeys::Up || KeyEvent.GetKey() == EKeys::Down)
 		{
-			if(KeyEvent.GetKey() == EKeys::Up)
-			{
-				if(SelectedSuggestion < 0)
-				{
-					// from edit control to end of list
-					SelectedSuggestion = Suggestions.Num() - 1;
-				}
-				else
-				{
-					// got one up, possibly back to edit control
-					--SelectedSuggestion;
-				}
-			}
-
-			if(KeyEvent.GetKey() == EKeys::Down)
-			{
-				if(SelectedSuggestion < Suggestions.Num() - 1)
-				{
-					// go one down, possibly from edit control to top
-					++SelectedSuggestion;
-				}
-				else
-				{
-					// back to edit control
-					SelectedSuggestion = -1;
-				}
-			}
-
+			Suggestions.StepSelectedSuggestion(KeyEvent.GetKey() == EKeys::Up ? -1 : +1);
 			MarkActiveSuggestion();
 
 			return FReply::Handled();
 		}
 		else if (KeyEvent.GetKey() == EKeys::Tab)
 		{
-			if (Suggestions.Num())
+			if (Suggestions.HasSuggestions())
 			{
-				if (SelectedSuggestion >= 0 && SelectedSuggestion < Suggestions.Num())
+				if (Suggestions.HasSelectedSuggestion())
 				{
 					MarkActiveSuggestion();
 					OnTextCommitted(InputText->GetText(), ETextCommit::OnEnter);
 				}
 				else
 				{
-					SelectedSuggestion = 0;
+					Suggestions.SelectedSuggestion = 0;
 					MarkActiveSuggestion();
 				}
 			}
 
 			return FReply::Handled();
 		}
+		else if (KeyEvent.GetKey() == EKeys::Escape)
+		{
+			SuggestionBox->SetIsOpen(false);
+			return FReply::Handled();
+		}
 	}
 	else
 	{
 		if(KeyEvent.GetKey() == EKeys::Up)
 		{
-			TArray<FString> History;
+			// If the command field isn't empty we need you to have pressed Control+Up to summon the history (to make sure you're not just using caret navigation)
+			const bool bIsMultiLine = GetActiveCommandExecutorAllowMultiLine();
+			const bool bShowHistory = InputText->GetText().IsEmpty() || KeyEvent.IsControlDown();
+			if (bShowHistory)
+			{
+				TArray<FString> History;
+				if (ActiveCommandExecutor)
+				{
+					ActiveCommandExecutor->GetExecHistory(History);
+				}
 
-			IConsoleManager::Get().GetConsoleHistory(History);
+				SetSuggestions(History, FText::GetEmpty());
+				
+				if(Suggestions.HasSuggestions())
+				{
+					Suggestions.StepSelectedSuggestion(-1);
+					MarkActiveSuggestion();
+				}
+			}
 
-			SetSuggestions(History, true);
-			
-			if(Suggestions.Num())
+			// Need to always handle this for single-line controls to avoid them invoking widget navigation
+			if (!bIsMultiLine || bShowHistory)
+			{
+				return FReply::Handled();
+			}
+		}
+		else if (KeyEvent.GetKey() == EKeys::Escape)
+		{
+			if (InputText->GetText().IsEmpty())
+			{
+				OnCloseConsole.ExecuteIfBound();
+			}
+			else
 			{
-				SelectedSuggestion = Suggestions.Num() - 1;
-				MarkActiveSuggestion();
+				// Clear the console input area
+				bIgnoreUIUpdate = true;
+				InputText->SetText(FText::GetEmpty());
+				bIgnoreUIUpdate = false;
+
+				ClearSuggestions();
 			}
 
 			return FReply::Handled();
@@ -522,33 +481,36 @@ FReply SConsoleInputBox::OnPreviewKeyDown(const FGeometry& MyGeometry, const FKe
 	return FReply::Unhandled();
 }
 
-void SConsoleInputBox::SetSuggestions(TArray<FString>& Elements, bool bInHistoryMode)
+void SConsoleInputBox::SetSuggestions(TArray<FString>& Elements, FText Highlight)
 {
 	FString SelectionText;
-	if (SelectedSuggestion >= 0 && SelectedSuggestion < Suggestions.Num())
+	if (Suggestions.HasSelectedSuggestion())
 	{
-		SelectionText = *Suggestions[SelectedSuggestion];
+		SelectionText = *Suggestions.GetSelectedSuggestion();
 	}
 
-	SelectedSuggestion = -1;
-	Suggestions.Empty();
-	SelectedSuggestion = -1;
+	Suggestions.Reset();
+	Suggestions.SuggestionsHighlight = Highlight;
 
-	for(uint32 i = 0; i < (uint32)Elements.Num(); ++i)
+	for(int32 i = 0; i < Elements.Num(); ++i)
 	{
-		Suggestions.Add(MakeShareable(new FString(Elements[i])));
+		Suggestions.SuggestionsList.Add(MakeShared<FString>(Elements[i]));
 
 		if (Elements[i] == SelectionText)
 		{
-			SelectedSuggestion = i;
+			Suggestions.SelectedSuggestion = i;
 		}
 	}
+	SuggestionListView->RequestListRefresh();
 
-	if(Suggestions.Num())
+	if(Suggestions.HasSuggestions())
 	{
 		// Ideally if the selection box is open the output window is not changing it's window title (flickers)
 		SuggestionBox->SetIsOpen(true, false);
-		SuggestionListView->RequestScrollIntoView(Suggestions.Last());
+		if (Suggestions.HasSelectedSuggestion())
+		{
+			SuggestionListView->RequestScrollIntoView(Suggestions.GetSelectedSuggestion());
+		}
 	}
 	else
 	{
@@ -564,12 +526,14 @@ void SConsoleInputBox::OnFocusLost( const FFocusEvent& InFocusEvent )
 void SConsoleInputBox::MarkActiveSuggestion()
 {
 	bIgnoreUIUpdate = true;
-	if(SelectedSuggestion >= 0)
+	if (Suggestions.HasSelectedSuggestion())
 	{
-		SuggestionListView->SetSelection(Suggestions[SelectedSuggestion]);
-		SuggestionListView->RequestScrollIntoView(Suggestions[SelectedSuggestion]);	// Ideally this would only scroll if outside of the view
+		TSharedPtr<FString> SelectedSuggestion = Suggestions.GetSelectedSuggestion();
+
+		SuggestionListView->SetSelection(SelectedSuggestion);
+		SuggestionListView->RequestScrollIntoView(SelectedSuggestion);	// Ideally this would only scroll if outside of the view
 
-		InputText->SetText(FText::FromString(GetSelectionText()));
+		InputText->SetText(FText::FromString(*SelectedSuggestion));
 	}
 	else
 	{
@@ -580,18 +544,172 @@ void SConsoleInputBox::MarkActiveSuggestion()
 
 void SConsoleInputBox::ClearSuggestions()
 {
-	SelectedSuggestion = -1;
 	SuggestionBox->SetIsOpen(false);
-	Suggestions.Empty();
+	Suggestions.Reset();
+}
+
+void SConsoleInputBox::OnCommandExecutorRegistered(const FName& Type, class IModularFeature* ModularFeature)
+{
+	if (Type == IConsoleCommandExecutor::ModularFeatureName())
+	{
+		SyncActiveCommandExecutor();
+	}
 }
 
-FString SConsoleInputBox::GetSelectionText() const
+void SConsoleInputBox::OnCommandExecutorUnregistered(const FName& Type, class IModularFeature* ModularFeature)
 {
-	FString ret = *Suggestions[SelectedSuggestion];
+	if (Type == IConsoleCommandExecutor::ModularFeatureName() && ModularFeature == ActiveCommandExecutor)
+	{
+		SyncActiveCommandExecutor();
+	}
+}
+
+void SConsoleInputBox::SyncActiveCommandExecutor()
+{
+	TArray<IConsoleCommandExecutor*> CommandExecutors = IModularFeatures::Get().GetModularFeatureImplementations<IConsoleCommandExecutor>(IConsoleCommandExecutor::ModularFeatureName());
+	ActiveCommandExecutor = nullptr;
+
+	// First try and match from the active name
+	for (IConsoleCommandExecutor* CommandExecutor : CommandExecutors)
+	{
+		if (CommandExecutor->GetName() == PreferredCommandExecutorName)
+		{
+			ActiveCommandExecutor = CommandExecutor;
+			break;
+		}
+	}
+
+	// Failing that, fallback to the default name
+	if (!ActiveCommandExecutor && PreferredCommandExecutorName != FConsoleCommandExecutor::StaticName())
+	{
+		for (IConsoleCommandExecutor* CommandExecutor : CommandExecutors)
+		{
+			if (CommandExecutor->GetName() == FConsoleCommandExecutor::StaticName())
+			{
+				ActiveCommandExecutor = CommandExecutor;
+				break;
+			}
+		}
+	}
+}
+
+void SConsoleInputBox::SetActiveCommandExecutor(const FName InExecName)
+{
+	GConfig->SetString(TEXT("OutputLog"), TEXT("PreferredCommandExecutor"), *InExecName.ToString(), GEditorPerProjectIni);
+	PreferredCommandExecutorName = InExecName;
+	SyncActiveCommandExecutor();
+}
+
+FText SConsoleInputBox::GetActiveCommandExecutorDisplayName() const
+{
+	if (ActiveCommandExecutor)
+	{
+		return ActiveCommandExecutor->GetDisplayName();
+	}
+	return FText::GetEmpty();
+}
+
+FText SConsoleInputBox::GetActiveCommandExecutorHintText() const
+{
+	if (ActiveCommandExecutor)
+	{
+		return ActiveCommandExecutor->GetHintText();
+	}
+	return FText::GetEmpty();
+}
+
+bool SConsoleInputBox::GetActiveCommandExecutorAllowMultiLine() const
+{
+	if (ActiveCommandExecutor)
+	{
+		return ActiveCommandExecutor->AllowMultiLine();
+	}
+	return false;
+}
+
+bool SConsoleInputBox::IsCommandExecutorMenuEnabled() const
+{
+	return !ConsoleCommandCustomExec.IsBound(); // custom execs always show the default executor in the UI (which has the selector disabled)
+}
+
+TSharedRef<SWidget> SConsoleInputBox::GetCommandExecutorMenuContent()
+{
+	FMenuBuilder MenuBuilder(/*bInShouldCloseWindowAfterMenuSelection=*/true, nullptr);
+
+	MenuBuilder.BeginSection("CmdExecEntries");
+	{
+		TArray<IConsoleCommandExecutor*> CommandExecutors = IModularFeatures::Get().GetModularFeatureImplementations<IConsoleCommandExecutor>(IConsoleCommandExecutor::ModularFeatureName());
+		CommandExecutors.Sort([](IConsoleCommandExecutor& LHS, IConsoleCommandExecutor& RHS)
+		{
+			return LHS.GetDisplayName().CompareTo(RHS.GetDisplayName()) < 0;
+		});
+
+		for (const IConsoleCommandExecutor* CommandExecutor : CommandExecutors)
+		{
+			const bool bIsActiveCmdExec = ActiveCommandExecutor == CommandExecutor;
+
+			MenuBuilder.AddMenuEntry(
+				CommandExecutor->GetDisplayName(),
+				CommandExecutor->GetDescription(),
+				FSlateIcon(),
+				FUIAction(
+					FExecuteAction::CreateSP(this, &SConsoleInputBox::SetActiveCommandExecutor, CommandExecutor->GetName()),
+					FCanExecuteAction::CreateLambda([] { return true; }),
+					FIsActionChecked::CreateLambda([bIsActiveCmdExec] { return bIsActiveCmdExec; })
+					),
+				NAME_None,
+				EUserInterfaceActionType::RadioButton
+			);
+		}
+	}
+	MenuBuilder.EndSection();
+
+	return MenuBuilder.MakeWidget();
+}
+
+FReply SConsoleInputBox::OnKeyDownHandler(const FGeometry& MyGeometry, const FKeyEvent& InKeyEvent)
+{
+	const FInputChord OpenConsoleChord = *FGlobalEditorCommonCommands::Get().OpenConsoleCommandBox->GetActiveChord(EMultipleKeyBindingIndex::Primary);
+	const FInputChord InputChord = FInputChord(InKeyEvent.GetKey(), EModifierKey::FromBools(InKeyEvent.IsControlDown(), InKeyEvent.IsAltDown(), InKeyEvent.IsShiftDown(), InKeyEvent.IsCommandDown()));
+
+	// Intercept the "open console" key
+	if ((!ActiveCommandExecutor || ActiveCommandExecutor->AllowHotKeyClose()) 
+		&& OpenConsoleChord == InputChord)
+	{
+		OnCloseConsole.ExecuteIfBound();
+		return FReply::Handled();
+	}
 	
-	ret = ret.Replace(TEXT("\t"), TEXT(""));
+	return FReply::Unhandled();
+}
+
+FReply SConsoleInputBox::OnKeyCharHandler(const FGeometry& MyGeometry, const FCharacterEvent& InCharacterEvent)
+{
+	// A printable key may be used to open the console, so consume all characters before our first Tick
+	if (!bHasTicked)
+	{
+		return FReply::Handled();
+	}
+
+	const FInputChord OpenConsoleChord = *FGlobalEditorCommonCommands::Get().OpenConsoleCommandBox->GetActiveChord(EMultipleKeyBindingIndex::Primary);
+
+	const uint32* KeyCode = nullptr;
+	const uint32* CharCode = nullptr;
+	FInputKeyManager::Get().GetCodesFromKey(OpenConsoleChord.Key, KeyCode, CharCode);
+	check(CharCode);
 
-	return ret;
+	// Intercept the "open console" key
+	if ((!ActiveCommandExecutor || ActiveCommandExecutor->AllowHotKeyClose())
+		&& InCharacterEvent.GetCharacter() == (TCHAR)*CharCode
+		&& OpenConsoleChord.NeedsControl() == InCharacterEvent.IsControlDown()
+		&& OpenConsoleChord.NeedsAlt() == InCharacterEvent.IsAltDown()
+		&& OpenConsoleChord.NeedsShift() == InCharacterEvent.IsShiftDown()
+		&& OpenConsoleChord.NeedsCommand() == InCharacterEvent.IsCommandDown())
+	{
+		return FReply::Handled();
+	}
+	
+	return FReply::Unhandled();
 }
 
 TSharedRef< FOutputLogTextLayoutMarshaller > FOutputLogTextLayoutMarshaller::Create(TArray< TSharedPtr<FLogMessage> > InMessages, FLogFilter* InFilter)
@@ -796,89 +914,87 @@ void SOutputLog::Construct( const FArguments& InArgs )
 
 	ChildSlot
 	[
-		SNew(SVerticalBox)
-
-		// Console output and filters
-		+SVerticalBox::Slot()
+		SNew(SBorder)
+		.Padding(3)
+		.BorderImage(FEditorStyle::GetBrush("ToolPanel.GroupBorder"))
 		[
-			SNew(SBorder)
-			.Padding(3)
-			.BorderImage(FEditorStyle::GetBrush("ToolPanel.GroupBorder"))
-			[
-				SNew(SVerticalBox)
+			SNew(SVerticalBox)
 
-				// Output Log Filter
-				+SVerticalBox::Slot()
-				.AutoHeight()
-				.Padding(FMargin(0.0f, 0.0f, 0.0f, 4.0f))
-				[
-					SNew(SHorizontalBox)
+			// Output Log Filter
+			+SVerticalBox::Slot()
+			.AutoHeight()
+			.Padding(FMargin(0.0f, 0.0f, 0.0f, 4.0f))
+			[
+				SNew(SHorizontalBox)
 			
-					+SHorizontalBox::Slot()
-					.AutoWidth()
+				+SHorizontalBox::Slot()
+				.AutoWidth()
+				[
+					SNew(SComboButton)
+					.ComboButtonStyle(FEditorStyle::Get(), "GenericFilters.ComboButtonStyle")
+					.ForegroundColor(FLinearColor::White)
+					.ContentPadding(0)
+					.ToolTipText(LOCTEXT("AddFilterToolTip", "Add an output log filter."))
+					.OnGetMenuContent(this, &SOutputLog::MakeAddFilterMenu)
+					.HasDownArrow(true)
+					.ContentPadding(FMargin(1, 0))
+					.ButtonContent()
 					[
-						SNew(SComboButton)
-						.ComboButtonStyle(FEditorStyle::Get(), "GenericFilters.ComboButtonStyle")
-						.ForegroundColor(FLinearColor::White)
-						.ContentPadding(0)
-						.ToolTipText(LOCTEXT("AddFilterToolTip", "Add an output log filter."))
-						.OnGetMenuContent(this, &SOutputLog::MakeAddFilterMenu)
-						.HasDownArrow(true)
-						.ContentPadding(FMargin(1, 0))
-						.ButtonContent()
+						SNew(SHorizontalBox)
+
+						+SHorizontalBox::Slot()
+						.AutoWidth()
 						[
-							SNew(SHorizontalBox)
-
-							+SHorizontalBox::Slot()
-							.AutoWidth()
-							[
-								SNew(STextBlock)
-								.TextStyle(FEditorStyle::Get(), "GenericFilters.TextStyle")
-								.Font(FEditorStyle::Get().GetFontStyle("FontAwesome.9"))
-								.Text(FText::FromString(FString(TEXT("\xf0b0"))) /*fa-filter*/)
-							]
-
-							+SHorizontalBox::Slot()
-							.AutoWidth()
-							.Padding(2, 0, 0, 0)
-							[
-								SNew(STextBlock)
-								.TextStyle(FEditorStyle::Get(), "GenericFilters.TextStyle")
-								.Text(LOCTEXT("Filters", "Filters"))
-							]
+							SNew(STextBlock)
+							.TextStyle(FEditorStyle::Get(), "GenericFilters.TextStyle")
+							.Font(FEditorStyle::Get().GetFontStyle("FontAwesome.9"))
+							.Text(FText::FromString(FString(TEXT("\xf0b0"))) /*fa-filter*/)
 						]
-					]
 
-					+SHorizontalBox::Slot()
-					.Padding(4, 1, 0, 0)
-					[
-						SAssignNew(FilterTextBox, SSearchBox)
-						.HintText(LOCTEXT("SearchLogHint", "Search Log"))
-						.OnTextChanged(this, &SOutputLog::OnFilterTextChanged)
-						.OnTextCommitted(this, &SOutputLog::OnFilterTextCommitted)
-						.DelayChangeNotificationsWhileTyping(true)
+						+SHorizontalBox::Slot()
+						.AutoWidth()
+						.Padding(2, 0, 0, 0)
+						[
+							SNew(STextBlock)
+							.TextStyle(FEditorStyle::Get(), "GenericFilters.TextStyle")
+							.Text(LOCTEXT("Filters", "Filters"))
+						]
 					]
 				]
 
-				// Output log area
-				+SVerticalBox::Slot()
-				.FillHeight(1)
+				+SHorizontalBox::Slot()
+				.Padding(4, 1, 0, 0)
 				[
-					MessagesTextBox.ToSharedRef()
+					SAssignNew(FilterTextBox, SSearchBox)
+					.HintText(LOCTEXT("SearchLogHint", "Search Log"))
+					.OnTextChanged(this, &SOutputLog::OnFilterTextChanged)
+					.OnTextCommitted(this, &SOutputLog::OnFilterTextCommitted)
+					.DelayChangeNotificationsWhileTyping(true)
 				]
 			]
-		]
 
-		// The console input box
-		+SVerticalBox::Slot()
-		.AutoHeight()
-		.Padding(FMargin(0.0f, 4.0f, 0.0f, 0.0f))
-		[
-			SNew(SConsoleInputBox)
-			.OnConsoleCommandExecuted(this, &SOutputLog::OnConsoleCommandExecuted)
+			// Output log area
+			+SVerticalBox::Slot()
+			.FillHeight(1)
+			[
+				MessagesTextBox.ToSharedRef()
+			]
 
-			// Always place suggestions above the input line for the output log widget
-			.SuggestionListPlacement(MenuPlacement_AboveAnchor)
+			// The console input box
+			+SVerticalBox::Slot()
+			.AutoHeight()
+			.Padding(FMargin(0.0f, 4.0f, 0.0f, 0.0f))
+			[
+				SNew(SBox)
+				.MaxDesiredHeight(180.0f)
+				[
+					SNew(SConsoleInputBox)
+					.OnConsoleCommandExecuted(this, &SOutputLog::OnConsoleCommandExecuted)
+
+					// Always place suggestions above the input line for the output log widget
+					.SuggestionListPlacement(MenuPlacement_AboveAnchor)
+				]
+			]
 		]
 	];
 
diff --git a/Engine/Source/Developer/OutputLog/Private/SOutputLog.h b/Engine/Source/Developer/OutputLog/Private/SOutputLog.h
--- a/Engine/Source/Developer/OutputLog/Private/SOutputLog.h
+++ b/Engine/Source/Developer/OutputLog/Private/SOutputLog.h
@@ -8,17 +8,35 @@
 #include "Input/Reply.h"
 #include "Widgets/SWidget.h"
 #include "Widgets/SCompoundWidget.h"
-#include "Widgets/Input/SEditableTextBox.h"
+#include "Widgets/Input/SMultiLineEditableTextBox.h"
 #include "Widgets/Views/STableViewBase.h"
 #include "Widgets/Views/STableRow.h"
 #include "Framework/Text/BaseTextLayoutMarshaller.h"
 #include "Misc/TextFilterExpressionEvaluator.h"
+#include "HAL/IConsoleManager.h"
 
 class FMenuBuilder;
 class FOutputLogTextLayoutMarshaller;
 class FTextLayout;
 class SMenuAnchor;
-class SMultiLineEditableTextBox;
+
+/**
+ * Executor for Unreal console commands
+ */
+class FConsoleCommandExecutor : public IConsoleCommandExecutor
+{
+public:
+	static FName StaticName();
+	virtual FName GetName() const override;
+	virtual FText GetDisplayName() const override;
+	virtual FText GetDescription() const override;
+	virtual FText GetHintText() const override;
+	virtual void GetAutoCompleteSuggestions(const TCHAR* Input, TArray<FString>& Out) override;
+	virtual void GetExecHistory(TArray<FString>& Out) override;
+	virtual bool Exec(const TCHAR* Input) override;
+	virtual bool AllowHotKeyClose() const override;
+	virtual bool AllowMultiLine() const override;
+};
 
 /**
 * A single log message for the output log, holding a message and
@@ -70,6 +88,9 @@ public:
 
 		/** Called when a console command is executed */
 		SLATE_EVENT( FSimpleDelegate, OnConsoleCommandExecuted )
+
+		/** Delegate to call to close the console */
+		SLATE_EVENT( FSimpleDelegate, OnCloseConsole )
 	SLATE_END_ARGS()
 
 	/** Protected console input box widget constructor, called by Slate */
@@ -83,7 +104,7 @@ public:
 	void Construct( const FArguments& InArgs );
 
 	/** Returns the editable text box associated with this widget.  Used to set focus directly. */
-	TSharedRef< SEditableTextBox > GetEditableTextBox()
+	TSharedRef< SMultiLineEditableTextBox > GetEditableTextBox()
 	{
 		return InputText.ToSharedRef();
 	}
@@ -110,39 +131,120 @@ protected:
 
 	void SuggestionSelectionChanged(TSharedPtr<FString> NewValue, ESelectInfo::Type SelectInfo);
 		
-	void SetSuggestions(TArray<FString>& Elements, bool bInHistoryMode);
+	void SetSuggestions(TArray<FString>& Elements, FText Highlight);
 
 	void MarkActiveSuggestion();
 
 	void ClearSuggestions();
 
-	FString GetSelectionText() const;
+	void OnCommandExecutorRegistered(const FName& Type, class IModularFeature* ModularFeature);
+
+	void OnCommandExecutorUnregistered(const FName& Type, class IModularFeature* ModularFeature);
+
+	void SyncActiveCommandExecutor();
+
+	void SetActiveCommandExecutor(const FName InExecName);
+
+	FText GetActiveCommandExecutorDisplayName() const;
+
+	FText GetActiveCommandExecutorHintText() const;
+
+	bool GetActiveCommandExecutorAllowMultiLine() const;
+
+	bool IsCommandExecutorMenuEnabled() const;
+
+	TSharedRef<SWidget> GetCommandExecutorMenuContent();
+
+	FReply OnKeyDownHandler(const FGeometry& MyGeometry, const FKeyEvent& InKeyEvent);
+
+	FReply OnKeyCharHandler(const FGeometry& MyGeometry, const FCharacterEvent& InCharacterEvent);
 
 private:
 
+	struct FSuggestions
+	{
+		FSuggestions()
+			: SelectedSuggestion(INDEX_NONE)
+		{
+		}
+
+		void Reset()
+		{
+			SelectedSuggestion = INDEX_NONE;
+			SuggestionsList.Reset();
+			SuggestionsHighlight = FText::GetEmpty();
+		}
+
+		bool HasSuggestions() const
+		{
+			return SuggestionsList.Num() > 0;
+		}
+
+		bool HasSelectedSuggestion() const
+		{
+			return SuggestionsList.IsValidIndex(SelectedSuggestion);
+		}
+
+		void StepSelectedSuggestion(const int32 Step)
+		{
+			SelectedSuggestion += Step;
+			if (SelectedSuggestion < 0)
+			{
+				SelectedSuggestion = SuggestionsList.Num() - 1;
+			}
+			else if (SelectedSuggestion >= SuggestionsList.Num())
+			{
+				SelectedSuggestion = 0;
+			}
+		}
+
+		TSharedPtr<FString> GetSelectedSuggestion() const
+		{
+			return SuggestionsList.IsValidIndex(SelectedSuggestion) ? SuggestionsList[SelectedSuggestion] : nullptr;
+		}
+
+		/** INDEX_NONE if not set, otherwise index into SuggestionsList */
+		int32 SelectedSuggestion;
+
+		/** All log messages stored in this widget for the list view */
+		TArray<TSharedPtr<FString>> SuggestionsList;
+
+		/** Highlight text to use for the suggestions list */
+		FText SuggestionsHighlight;
+	};
+
 	/** Editable text widget */
-	TSharedPtr< SEditableTextBox > InputText;
+	TSharedPtr< SMultiLineEditableTextBox > InputText;
 
 	/** history / auto completion elements */
 	TSharedPtr< SMenuAnchor > SuggestionBox;
 
-	/** All log messages stored in this widget for the list view */
-	TArray< TSharedPtr<FString> > Suggestions;
-
 	/** The list view for showing all log messages. Should be replaced by a full text editor */
 	TSharedPtr< SListView< TSharedPtr<FString> > > SuggestionListView;
 
+	/** Active list of suggestions */
+	FSuggestions Suggestions;
+
 	/** Delegate to call when a console command is executed */
 	FSimpleDelegate OnConsoleCommandExecuted;
 
 	/** Delegate to call to execute console command */
 	FExecuteConsoleCommand ConsoleCommandCustomExec;
 
-	/** -1 if not set, otherwise index into Suggestions */
-	int32 SelectedSuggestion;
+	/** Delegate to call to close the console */
+	FSimpleDelegate OnCloseConsole;
+
+	/** Name of the preferred command executor (may not always be the active executor) */
+	FName PreferredCommandExecutorName;
+
+	/** The currently active command executor */
+	IConsoleCommandExecutor* ActiveCommandExecutor;
 
 	/** to prevent recursive calls in UI callback */
-	bool bIgnoreUIUpdate; 
+	bool bIgnoreUIUpdate;
+
+	/** true if this widget has been Ticked at least once */
+	bool bHasTicked;
 };
 
 /**
diff --git a/Engine/Source/Developer/OutputLog/Public/OutputLogModule.h b/Engine/Source/Developer/OutputLog/Public/OutputLogModule.h
--- a/Engine/Source/Developer/OutputLog/Public/OutputLogModule.h
+++ b/Engine/Source/Developer/OutputLog/Public/OutputLogModule.h
@@ -7,7 +7,8 @@
 #include "Modules/ModuleInterface.h"
 #include "Widgets/SWindow.h"
 
-class SEditableTextBox;
+class FConsoleCommandExecutor;
+class SMultiLineEditableTextBox;
 
 /** Style of the debug console */
 namespace EDebugConsoleStyle
@@ -26,6 +27,7 @@ struct FDebugConsoleDelegates
 {
 	FSimpleDelegate OnFocusLost;
 	FSimpleDelegate OnConsoleCommandExecuted;
+	FSimpleDelegate OnCloseConsole;
 };
 
 class FOutputLogModule : public IModuleInterface
@@ -36,7 +38,7 @@ public:
 
 	/** Generates a console input box widget.  Remember, this widget will become invalid if the
 	    output log DLL is unloaded on the fly. */
-	virtual TSharedRef< SWidget > MakeConsoleInputBox( TSharedPtr< SEditableTextBox >& OutExposedEditableTextBox ) const;
+	virtual TSharedRef< SWidget > MakeConsoleInputBox( TSharedPtr< SMultiLineEditableTextBox >& OutExposedEditableTextBox ) const;
 
 	/** Opens a debug console in the specified window, if not already open */
 	virtual void ToggleDebugConsoleForWindow( const TSharedRef< SWindow >& Window, const EDebugConsoleStyle::Type InStyle, const FDebugConsoleDelegates& DebugConsoleDelegates );
@@ -49,4 +51,7 @@ private:
 
 	/** Weak pointer to a debug console that's currently open, if any */
 	TWeakPtr< SWidget > DebugConsole;
+
+	/** Pointer to the classic "Cmd" executor */
+	TSharedPtr<FConsoleCommandExecutor> CmdExec;
 };
diff --git a/Engine/Source/Editor/AddContentDialog/Private/ContentSourceProviders/FeaturePack/FeaturePackContentSourceProvider.cpp b/Engine/Source/Editor/AddContentDialog/Private/ContentSourceProviders/FeaturePack/FeaturePackContentSourceProvider.cpp
--- a/Engine/Source/Editor/AddContentDialog/Private/ContentSourceProviders/FeaturePack/FeaturePackContentSourceProvider.cpp
+++ b/Engine/Source/Editor/AddContentDialog/Private/ContentSourceProviders/FeaturePack/FeaturePackContentSourceProvider.cpp
@@ -33,7 +33,9 @@ public:
 FFeaturePackContentSourceProvider::FFeaturePackContentSourceProvider()
 {
 	FeaturePackPath = FPaths::FeaturePackDir();
+	EnterpriseFeaturePackPath = FPaths::EnterpriseFeaturePackDir();
 	TemplatePath = FPaths::RootDir() + TEXT("Templates/");
+	EnterpriseTemplatePath = FPaths::EnterpriseDir() + TEXT("Templates/");
 	StartUpDirectoryWatcher();
 	RefreshFeaturePacks();
 }
@@ -58,6 +60,13 @@ void FFeaturePackContentSourceProvider::StartUpDirectoryWatcher()
 		IFileManager::Get().MakeDirectory(*FeaturePackPath);
 		DirectoryChangedDelegate = IDirectoryWatcher::FDirectoryChanged::CreateRaw(this, &FFeaturePackContentSourceProvider::OnFeaturePackDirectoryChanged);
 		DirectoryWatcher->RegisterDirectoryChangedCallback_Handle(FeaturePackPath, DirectoryChangedDelegate, DirectoryChangedDelegateHandle);
+
+		if (IFileManager::Get().DirectoryExists(*FPaths::EnterpriseDir()))
+		{
+			// If the path doesn't exist on disk, make it so the watcher will work.
+			IFileManager::Get().MakeDirectory(*EnterpriseFeaturePackPath);
+			DirectoryWatcher->RegisterDirectoryChangedCallback_Handle(EnterpriseFeaturePackPath, DirectoryChangedDelegate, DirectoryChangedDelegateHandle);
+		}
 	}
 }
 
@@ -68,6 +77,7 @@ void FFeaturePackContentSourceProvider::ShutDownDirectoryWatcher()
 	if ( DirectoryWatcher )
 	{
 		DirectoryWatcher->UnregisterDirectoryChangedCallback_Handle(FeaturePackPath, DirectoryChangedDelegateHandle);
+		DirectoryWatcher->UnregisterDirectoryChangedCallback_Handle(EnterpriseFeaturePackPath, DirectoryChangedDelegateHandle);
 	}
 }
 
@@ -89,7 +99,8 @@ void FFeaturePackContentSourceProvider::RefreshFeaturePacks()
 	// first the .upack files
 	IPlatformFile &PlatformFile = FPlatformFileManager::Get().GetPlatformFile();
 	FFillArrayDirectoryVisitor DirectoryVisitor;
-	PlatformFile.IterateDirectory( *FeaturePackPath, DirectoryVisitor );
+	PlatformFile.IterateDirectory(*FeaturePackPath, DirectoryVisitor);
+	PlatformFile.IterateDirectory(*EnterpriseFeaturePackPath, DirectoryVisitor);
 	for ( auto FeaturePackFile : DirectoryVisitor.Files )
 	{
 		if( FeaturePackFile.EndsWith(TEXT(".upack")) == true)
@@ -105,6 +116,7 @@ void FFeaturePackContentSourceProvider::RefreshFeaturePacks()
 	// Now the 'loose' feature packs
 	FFillArrayDirectoryVisitor TemplateDirectoryVisitor;
 	PlatformFile.IterateDirectoryRecursively(*TemplatePath, TemplateDirectoryVisitor);
+	PlatformFile.IterateDirectoryRecursively(*EnterpriseTemplatePath, TemplateDirectoryVisitor);
 	for (auto TemplatePackFile: TemplateDirectoryVisitor.Files)
 	{
 		FString ThisPackRoot = FPaths::GetPath(TemplatePackFile);
diff --git a/Engine/Source/Editor/AddContentDialog/Private/ContentSourceProviders/FeaturePack/FeaturePackContentSourceProvider.h b/Engine/Source/Editor/AddContentDialog/Private/ContentSourceProviders/FeaturePack/FeaturePackContentSourceProvider.h
--- a/Engine/Source/Editor/AddContentDialog/Private/ContentSourceProviders/FeaturePack/FeaturePackContentSourceProvider.h
+++ b/Engine/Source/Editor/AddContentDialog/Private/ContentSourceProviders/FeaturePack/FeaturePackContentSourceProvider.h
@@ -33,9 +33,15 @@ private:
 	/** The path on disk to the directory containing the feature packs. */
 	FString FeaturePackPath;
 
+	/** The path on disk to the directory containing the enterprise feature packs. */
+	FString EnterpriseFeaturePackPath;
+
 	/** The path on disk to the directory containing the templates folder. */
 	FString TemplatePath;
 
+	/** The path on disk to the directory containing the enterprise templates folder. */
+	FString EnterpriseTemplatePath;
+
 	/** The delegate which gets called when the feature pack directory changes. This reference
 		is kept so that it can be unregistered correctly. */
 	IDirectoryWatcher::FDirectoryChanged DirectoryChangedDelegate;
@@ -46,5 +52,8 @@ private:
 	/** An array of the available content sources. */
 	TArray<TSharedRef<IContentSource>> ContentSources;
 
+	/** An array of the available Enterprise content sources. */
+	TArray<TSharedRef<IContentSource>> EnterpiseContentSources;
+
 	FDelegateHandle DirectoryChangedDelegateHandle;
 };
diff --git a/Engine/Source/Editor/AddContentDialog/Private/IContentSource.h b/Engine/Source/Editor/AddContentDialog/Private/IContentSource.h
--- a/Engine/Source/Editor/AddContentDialog/Private/IContentSource.h
+++ b/Engine/Source/Editor/AddContentDialog/Private/IContentSource.h
@@ -12,7 +12,9 @@ enum class EContentSourceCategory:uint8
 {
 	BlueprintFeature,
 	CodeFeature,
+	EnterpriseFeature,
 	Content,
+	EnterpriseContent,
 	SharedPack,
 	Unknown
 };
diff --git a/Engine/Source/Editor/AddContentDialog/Private/ViewModels/CategoryViewModel.cpp b/Engine/Source/Editor/AddContentDialog/Private/ViewModels/CategoryViewModel.cpp
--- a/Engine/Source/Editor/AddContentDialog/Private/ViewModels/CategoryViewModel.cpp
+++ b/Engine/Source/Editor/AddContentDialog/Private/ViewModels/CategoryViewModel.cpp
@@ -47,15 +47,25 @@ void FCategoryViewModel::Initialize()
 		IconBrush = FAddContentDialogStyle::Get().GetBrush("AddContentDialog.CodeFeatureCategory");
 		SortID = 1;
 		break;
+	case EContentSourceCategory::EnterpriseFeature:
+		Text = LOCTEXT("EnterpriseFeature", "Enterprise Feature");
+		IconBrush = FAddContentDialogStyle::Get().GetBrush("AddContentDialog.BlueprintFeatureCategory");
+		SortID = 2;
+		break;
 	case EContentSourceCategory::Content:
 		Text = LOCTEXT("ContentPacks", "Content Packs");
 		IconBrush = FAddContentDialogStyle::Get().GetBrush("AddContentDialog.ContentPackCategory");
-		SortID = 2;
+		SortID = 3;
+		break;
+	case EContentSourceCategory::EnterpriseContent:
+		Text = LOCTEXT("EnterpriseContentPacks", "Enterprise Content Packs");
+		IconBrush = FAddContentDialogStyle::Get().GetBrush("AddContentDialog.ContentPackCategory");
+		SortID = 4;
 		break;
 	default:
 		Text = LOCTEXT("Miscellaneous", "Miscellaneous");
 		IconBrush = FAddContentDialogStyle::Get().GetBrush("AddContentDialog.UnknownCategory");
-		SortID = 3;
+		SortID = 5;
 		break;
 	}
 }
diff --git a/Engine/Source/Editor/EditorWidgets/Private/STextPropertyEditableTextBox.cpp b/Engine/Source/Editor/EditorWidgets/Private/STextPropertyEditableTextBox.cpp
--- a/Engine/Source/Editor/EditorWidgets/Private/STextPropertyEditableTextBox.cpp
+++ b/Engine/Source/Editor/EditorWidgets/Private/STextPropertyEditableTextBox.cpp
@@ -451,7 +451,7 @@ void STextPropertyEditableTextBox::Construct(const FArguments& InArgs, const TSh
 					.AutoWrapText(InArgs._AutoWrapText)
 					.WrapTextAt(InArgs._WrapTextAt)
 					.ModiferKeyForNewLine(EModifierKey::Shift)
-					.IsPassword(bIsPassword)
+					//.IsPassword(bIsPassword)
 				]
 			]
 		];
diff --git a/Engine/Source/Editor/GameProjectGeneration/Private/GameProjectGenerationModule.cpp b/Engine/Source/Editor/GameProjectGeneration/Private/GameProjectGenerationModule.cpp
--- a/Engine/Source/Editor/GameProjectGeneration/Private/GameProjectGenerationModule.cpp
+++ b/Engine/Source/Editor/GameProjectGeneration/Private/GameProjectGenerationModule.cpp
@@ -20,6 +20,7 @@ DEFINE_LOG_CATEGORY(LogGameProjectGeneration);
 
 FName FTemplateCategory::BlueprintCategoryName = "Blueprint";
 FName FTemplateCategory::CodeCategoryName = "C++";
+FName FTemplateCategory::EnterpriseCategoryName = "Enterprise";
 
 void FGameProjectGenerationModule::StartupModule()
 {
@@ -39,6 +40,13 @@ void FGameProjectGenerationModule::StartupModule()
 		),
 		FEditorStyle::GetBrush("GameProjectDialog.CodeIcon"),
 		FEditorStyle::GetBrush("GameProjectDialog.CodeImage"));
+
+	RegisterTemplateCategory(
+		FTemplateCategory::EnterpriseCategoryName,
+		LOCTEXT("EnterpriseCategory_Name", "Enterprise"),
+		LOCTEXT("EnterpriseCategory_Description", "Enterprise blueprint templates require no programming knowledge.\nEach template include a basic set of blueprints to use as a starting point for your enterprise project."),
+		FEditorStyle::GetBrush("GameProjectDialog.BlueprintIcon"),
+		FEditorStyle::GetBrush("GameProjectDialog.BlueprintImage"));
 }
 
 
diff --git a/Engine/Source/Editor/GameProjectGeneration/Private/SNewProjectWizard.cpp b/Engine/Source/Editor/GameProjectGeneration/Private/SNewProjectWizard.cpp
--- a/Engine/Source/Editor/GameProjectGeneration/Private/SNewProjectWizard.cpp
+++ b/Engine/Source/Editor/GameProjectGeneration/Private/SNewProjectWizard.cpp
@@ -303,6 +303,8 @@ private:
 BEGIN_SLATE_FUNCTION_BUILD_OPTIMIZATION
 void SNewProjectWizard::Construct( const FArguments& InArgs )
 {
+	bProjectSettingsHidden = false;
+
 	LastValidityCheckTime = 0;
 	ValidityCheckFrequency = 4;
 	bLastGlobalValidityCheckSuccessful = true;
@@ -310,8 +312,6 @@ void SNewProjectWizard::Construct( const FArguments& InArgs )
 	bPreventPeriodicValidityChecksUntilNextChange = false;
 	bCopyStarterContent = GEditor ? GetDefault<UEditorSettings>()->bCopyStarterContentPreference : true;
 
-	IHardwareTargetingModule& HardwareTargeting = IHardwareTargetingModule::Get();
-
 	SelectedHardwareClassTarget = EHardwareClass::Desktop;
 	SelectedGraphicsPreset = EGraphicsPreset::Maximum;
 
@@ -334,22 +334,6 @@ void SNewProjectWizard::Construct( const FArguments& InArgs )
 	TSharedRef<SSeparator> Separator = SNew(SSeparator).Orientation(EOrientation::Orient_Vertical);
 	Separator->SetBorderBackgroundColor(FLinearColor::White.CopyWithNewOpacity(0.25f));
 
-	TSharedPtr<SWidget> StartContentCombo;
-	{
-		TArray<SDecoratedEnumCombo<int32>::FComboOption> StarterContentInfo;
-		StarterContentInfo.Add(SDecoratedEnumCombo<int32>::FComboOption(
-			0, FSlateIcon(FEditorStyle::GetStyleSetName(), "GameProjectDialog.NoStarterContent"), LOCTEXT("NoStarterContent", "No Starter Content")));
-
-		// Only add the option to add starter content if its there to add !
-		bool bIsStarterAvailable = GameProjectUtils::IsStarterContentAvailableForNewProjects();
-		StarterContentInfo.Add(SDecoratedEnumCombo<int32>::FComboOption(
-			1, FSlateIcon(FEditorStyle::GetStyleSetName(), "GameProjectDialog.IncludeStarterContent"), LOCTEXT("IncludeStarterContent", "With Starter Content"),bIsStarterAvailable));
-		StartContentCombo = SNew(SDecoratedEnumCombo<int32>, MoveTemp(StarterContentInfo))
-			.SelectedEnum(this, &SNewProjectWizard::GetCopyStarterContentIndex)
-			.OnEnumChanged(this, &SNewProjectWizard::OnSetCopyStarterContent)
-			.ToolTipText(LOCTEXT("CopyStarterContent_ToolTip", "Enable to include an additional content pack containing simple placeable meshes with basic materials and textures.\nYou can opt out of including this to create a project that only has the bare essentials for the selected project template."));
-	}
-
 	const float UniformPadding = 16.f;
 	ChildSlot
 	[
@@ -572,11 +556,7 @@ void SNewProjectWizard::Construct( const FArguments& InArgs )
 								.AutoHeight()
 								.Padding(FMargin(0, 0, 0, 15.f))
 								[
-									SNew(SRichTextBlock)
-									.Text(LOCTEXT("ProjectSettingsDescription", "Choose some <RichTextBlock.BoldHighlight>settings</> for your project.  Don't worry, you can change these later in the <RichTextBlock.BoldHighlight>Target Hardware</> section of <RichTextBlock.BoldHighlight>Project Settings</>.  You can also add the <RichTextBlock.BoldHighlight>Starter Content</> to your project later using <RichTextBlock.BoldHighlight>Content Browser</>."))
-									.AutoWrapText(true)
-									.DecoratorStyleSet(&FEditorStyle::Get())
-									.ToolTip(IDocumentation::Get()->CreateToolTip(LOCTEXT("HardwareTargetTooltip", "These settings will choose good defaults for a number of other settings in the project such as post-processing flags and touch input emulation using the mouse."), NULL, TEXT("Shared/Editor/NewProjectWizard"), TEXT("TargetHardware")))
+									MakeProjectSettingsDescriptionBox()
 								]
 
 								+ SVerticalBox::Slot()
@@ -593,49 +573,7 @@ void SNewProjectWizard::Construct( const FArguments& InArgs )
 										.HAlign(HAlign_Center)
 										.Padding(FMargin(0, 0, 0, 25.f))
 										[
-											SNew(SHorizontalBox)
-
-											+ SHorizontalBox::Slot()
-											.AutoWidth()
-											[
-												HardwareTargeting.MakeHardwareClassTargetCombo(
-													FOnHardwareClassChanged::CreateSP(this, &SNewProjectWizard::SetHardwareClassTarget),
-													TAttribute<EHardwareClass::Type>(this, &SNewProjectWizard::GetHardwareClassTarget)
-													)
-											]
-								
-											+ SHorizontalBox::Slot()
-											.AutoWidth()
-											.Padding(FMargin(30, 0))
-											[
-												HardwareTargeting.MakeGraphicsPresetTargetCombo(
-													FOnGraphicsPresetChanged::CreateSP(this, &SNewProjectWizard::SetGraphicsPreset),
-													TAttribute<EGraphicsPreset::Type>(this, &SNewProjectWizard::GetGraphicsPreset)
-													)
-											]
-								
-											+ SHorizontalBox::Slot()
-											.AutoWidth()
-											[
-												SNew(SOverlay)
-												+SOverlay::Slot()
-												[
-													StartContentCombo.ToSharedRef()
-												]
-
-												// Warning when enabled for mobile, since the current starter content is bad for mobile
-												+SOverlay::Slot()
-												//.Visibility(EVisibility::SelfHitTestInvisible)
-												.HAlign(HAlign_Right)
-												.VAlign(VAlign_Top)
-												.Padding(4)
-												[
-													SNew(SImage)
-													.Image(FEditorStyle::GetBrush("Icons.Warning"))
-													.ToolTipText(this, &SNewProjectWizard::GetStarterContentWarningTooltip)
-													.Visibility(this, &SNewProjectWizard::GetStarterContentWarningVisibility)
-												]
-											]
+											MakeProjectSettingsOptionsBox()
 										]
 									]
 								]
@@ -1161,12 +1099,29 @@ TMap<FName, TArray<TSharedPtr<FTemplateItem>> >& SNewProjectWizard::FindTemplate
 		TEXT("")		// No asset types
 		)) );
 
+	if (FPaths::DirectoryExists(FPaths::EnterpriseDir()))
+	{
+		Templates.FindOrAdd(FTemplateCategory::EnterpriseCategoryName).Add(MakeShareable(new FTemplateItem(
+			LOCTEXT("BlankEnterpriseProjectName", "Blank Enterprise"),
+			LOCTEXT("BlankEnterpriseProjectDescription", "A clean empty enterprise project with no code."),
+			false, FTemplateCategory::BlueprintCategoryName,
+			TEXT("_1"),			// SortKey
+			TEXT(""),			// No filename, this is a generation template
+			MakeShareable(new FSlateBrush(*FEditorStyle::GetBrush("GameProjectDialog.BlankProjectThumbnail"))),
+			MakeShareable(new FSlateBrush(*FEditorStyle::GetBrush("GameProjectDialog.BlankProjectPreview"))),
+			TEXT(""),		// No class types
+			TEXT("")		// No asset types
+		)));
+	}
 	// Now discover and all data driven templates
 	TArray<FString> TemplateRootFolders;
 
 	// @todo rocket make template folder locations extensible.
 	TemplateRootFolders.Add( FPaths::RootDir() + TEXT("Templates") );
 
+	// Add the Enterprise templates
+	TemplateRootFolders.Add(FPaths::EnterpriseDir() + TEXT("Templates"));
+
 	// allow plugins to define templates
 	TArray<TSharedRef<IPlugin>> Plugins = IPluginManager::Get().GetEnabledPlugins();
 	for (const TSharedRef<IPlugin>& Plugin : Plugins)
@@ -1451,18 +1406,6 @@ bool SNewProjectWizard::CreateProject( const FString& ProjectFile )
 	ProjectInfo.TargetedHardware = SelectedHardwareClassTarget;
 	ProjectInfo.DefaultGraphicsPerformance = SelectedGraphicsPreset;
 
-	const FProjectDescriptor* CurrentProject = IProjectManager::Get().GetCurrentProject();
-	if (CurrentProject != nullptr)
-	{
-		ProjectInfo.bIsEnterpriseProject = CurrentProject->bIsEnterpriseProject;
-	}
-	else
-	{
-		// Set the default value for the enterprise flag from the command line for now.
-		// This should be temporary until we implement a more generic approach.
-		ProjectInfo.bIsEnterpriseProject = FParse::Param(FCommandLine::Get(), TEXT("enterprise"));
-	}
-
 	if (!GameProjectUtils::CreateProject(ProjectInfo, FailReason, FailLog))
 	{
 		SOutputLogDialog::Open(LOCTEXT("CreateProject", "Create Project"), FailReason, FailLog, FText::GetEmpty());
@@ -1634,6 +1577,21 @@ void SNewProjectWizard::HandleCategoryChanged(ECheckBoxState CheckState, FName C
 		TemplateListView->SetSelection(FilteredTemplateList[0]);
 	}
 	TemplateListView->RequestListRefresh();
+
+	if (ActiveCategory == FTemplateCategory::EnterpriseCategoryName)
+	{
+		bProjectSettingsHidden = true;
+		SetGraphicsPreset(EGraphicsPreset::Maximum);
+		SetHardwareClassTarget(EHardwareClass::Desktop);
+		ProjectSettingsDescriptionBox->SetVisibility(EVisibility::Collapsed);
+		ProjectSettingsOptionsBox->SetVisibility(EVisibility::Collapsed);
+	}
+	else if (bProjectSettingsHidden)
+	{
+		bProjectSettingsHidden = false;
+		ProjectSettingsDescriptionBox->SetVisibility(EVisibility::Visible);
+		ProjectSettingsOptionsBox->SetVisibility(EVisibility::Visible);
+	}
 }
 
 void SNewProjectWizard::SetHardwareClassTarget(EHardwareClass::Type InHardwareClass)
@@ -1646,4 +1604,85 @@ void SNewProjectWizard::SetGraphicsPreset(EGraphicsPreset::Type InGraphicsPreset
 	SelectedGraphicsPreset = InGraphicsPreset;
 }
 
+TSharedRef<SRichTextBlock> SNewProjectWizard::MakeProjectSettingsDescriptionBox()
+{
+	TSharedPtr<SRichTextBlock> Widget = SNew(SRichTextBlock)
+		.Text(LOCTEXT("ProjectSettingsDescription", "Choose some <RichTextBlock.BoldHighlight>settings</> for your project.  Don't worry, you can change these later in the <RichTextBlock.BoldHighlight>Target Hardware</> section of <RichTextBlock.BoldHighlight>Project Settings</>.  You can also add the <RichTextBlock.BoldHighlight>Starter Content</> to your project later using <RichTextBlock.BoldHighlight>Content Browser</>."))
+		.AutoWrapText(true)
+		.DecoratorStyleSet(&FEditorStyle::Get())
+		.ToolTip(IDocumentation::Get()->CreateToolTip(LOCTEXT("HardwareTargetTooltip", "These settings will choose good defaults for a number of other settings in the project such as post-processing flags and touch input emulation using the mouse."), NULL, TEXT("Shared/Editor/NewProjectWizard"), TEXT("TargetHardware")));
+
+	ProjectSettingsDescriptionBox = Widget;
+
+	return ProjectSettingsDescriptionBox.ToSharedRef();
+}
+
+TSharedRef<SHorizontalBox> SNewProjectWizard::MakeProjectSettingsOptionsBox()
+{
+	IHardwareTargetingModule& HardwareTargeting = IHardwareTargetingModule::Get();
+
+	TSharedPtr<SWidget> StartContentCombo;
+	{
+		TArray<SDecoratedEnumCombo<int32>::FComboOption> StarterContentInfo;
+		StarterContentInfo.Add(SDecoratedEnumCombo<int32>::FComboOption(
+			0, FSlateIcon(FEditorStyle::GetStyleSetName(), "GameProjectDialog.NoStarterContent"), LOCTEXT("NoStarterContent", "No Starter Content")));
+
+		// Only add the option to add starter content if its there to add !
+		bool bIsStarterAvailable = GameProjectUtils::IsStarterContentAvailableForNewProjects();
+		StarterContentInfo.Add(SDecoratedEnumCombo<int32>::FComboOption(
+			1, FSlateIcon(FEditorStyle::GetStyleSetName(), "GameProjectDialog.IncludeStarterContent"), LOCTEXT("IncludeStarterContent", "With Starter Content"), bIsStarterAvailable));
+		StartContentCombo = SNew(SDecoratedEnumCombo<int32>, MoveTemp(StarterContentInfo))
+			.SelectedEnum(this, &SNewProjectWizard::GetCopyStarterContentIndex)
+			.OnEnumChanged(this, &SNewProjectWizard::OnSetCopyStarterContent)
+			.ToolTipText(LOCTEXT("CopyStarterContent_ToolTip", "Enable to include an additional content pack containing simple placeable meshes with basic materials and textures.\nYou can opt out of including this to create a project that only has the bare essentials for the selected project template."));
+	}
+
+	TSharedPtr<SHorizontalBox> Widget = SNew(SHorizontalBox)
+
+		+ SHorizontalBox::Slot()
+		.AutoWidth()
+		[
+			HardwareTargeting.MakeHardwareClassTargetCombo(
+				FOnHardwareClassChanged::CreateSP(this, &SNewProjectWizard::SetHardwareClassTarget),
+				TAttribute<EHardwareClass::Type>(this, &SNewProjectWizard::GetHardwareClassTarget)
+			)
+		]
+
+		+ SHorizontalBox::Slot()
+		.AutoWidth()
+		.Padding(FMargin(30, 0))
+		[
+			HardwareTargeting.MakeGraphicsPresetTargetCombo(
+				FOnGraphicsPresetChanged::CreateSP(this, &SNewProjectWizard::SetGraphicsPreset),
+				TAttribute<EGraphicsPreset::Type>(this, &SNewProjectWizard::GetGraphicsPreset)
+			)
+		]
+
+		+ SHorizontalBox::Slot()
+		.AutoWidth()
+		[
+			SNew(SOverlay)
+			+ SOverlay::Slot()
+			[
+				StartContentCombo.ToSharedRef()
+			]
+
+			// Warning when enabled for mobile, since the current starter content is bad for mobile
+			+ SOverlay::Slot()
+			//.Visibility(EVisibility::SelfHitTestInvisible)
+			.HAlign(HAlign_Right)
+			.VAlign(VAlign_Top)
+			.Padding(4)
+			[
+				SNew(SImage)
+				.Image(FEditorStyle::GetBrush("Icons.Warning"))
+				.ToolTipText(this, &SNewProjectWizard::GetStarterContentWarningTooltip)
+				.Visibility(this, &SNewProjectWizard::GetStarterContentWarningVisibility)
+			]
+		];
+
+	ProjectSettingsOptionsBox = Widget;
+
+	return ProjectSettingsOptionsBox.ToSharedRef();
+}
 #undef LOCTEXT_NAMESPACE
diff --git a/Engine/Source/Editor/GameProjectGeneration/Private/SNewProjectWizard.h b/Engine/Source/Editor/GameProjectGeneration/Private/SNewProjectWizard.h
--- a/Engine/Source/Editor/GameProjectGeneration/Private/SNewProjectWizard.h
+++ b/Engine/Source/Editor/GameProjectGeneration/Private/SNewProjectWizard.h
@@ -8,6 +8,8 @@
 #include "Input/Reply.h"
 #include "Widgets/SWidget.h"
 #include "Widgets/SCompoundWidget.h"
+#include "Widgets/Text/SRichTextBlock.h"
+#include "Widgets/SBoxPanel.h"
 #include "HardwareTargetingSettings.h"
 #include "Widgets/Views/STileView.h"
 
@@ -208,6 +210,9 @@ private:
 	/** Get the check state for the specified categories tab */
 	ECheckBoxState GetCategoryTabCheckState(FName Category) const;
 
+	TSharedRef<SRichTextBlock> MakeProjectSettingsDescriptionBox();
+	TSharedRef<SHorizontalBox> MakeProjectSettingsOptionsBox();
+
 private:
 
 	/** The wizard widget */
@@ -260,4 +265,9 @@ private:
 	/** Names for pages */
 	static FName TemplatePageName;
 	static FName NameAndLocationPageName;
+
+	TSharedPtr<SRichTextBlock> ProjectSettingsDescriptionBox;
+	TSharedPtr<SHorizontalBox> ProjectSettingsOptionsBox;
+
+	bool bProjectSettingsHidden;
 };
diff --git a/Engine/Source/Editor/GameProjectGeneration/Private/TemplateCategory.h b/Engine/Source/Editor/GameProjectGeneration/Private/TemplateCategory.h
--- a/Engine/Source/Editor/GameProjectGeneration/Private/TemplateCategory.h
+++ b/Engine/Source/Editor/GameProjectGeneration/Private/TemplateCategory.h
@@ -8,7 +8,7 @@ struct FSlateBrush;
 
 struct FTemplateCategory
 {
-	static FName BlueprintCategoryName, CodeCategoryName;
+	static FName BlueprintCategoryName, CodeCategoryName, EnterpriseCategoryName;
 
 	/** Localised name of this category */
 	FText Name;
diff --git a/Engine/Source/Editor/GameProjectGeneration/Private/Tests/GameProjectAutomationTests.cpp b/Engine/Source/Editor/GameProjectGeneration/Private/Tests/GameProjectAutomationTests.cpp
--- a/Engine/Source/Editor/GameProjectGeneration/Private/Tests/GameProjectAutomationTests.cpp
+++ b/Engine/Source/Editor/GameProjectGeneration/Private/Tests/GameProjectAutomationTests.cpp
@@ -69,7 +69,7 @@ namespace GameProjectAutomationUtils
 
 
 	/* 
-	 * Create a project from a templalte with a given criteria
+	 * Create a project from a template with a given criteria
 	 * @oaram	InTemplates			List of available project templates
 	 * @param	InTargetedHardware	Target hardware (EHardwareClass)
 	 * @param	InGraphicPreset		Graphics preset (EGraphicsPreset)
@@ -99,6 +99,10 @@ namespace GameProjectAutomationUtils
 		{
 			CategoryName = FTemplateCategory::CodeCategoryName;
 		}
+		else if (InCategory == EContentSourceCategory::EnterpriseFeature)
+		{
+			CategoryName = FTemplateCategory::EnterpriseCategoryName;
+		}
 		else
 		{
 			// We didn't match a category
diff --git a/Engine/Source/Editor/LevelEditor/Private/LevelEditor.cpp b/Engine/Source/Editor/LevelEditor/Private/LevelEditor.cpp
--- a/Engine/Source/Editor/LevelEditor/Private/LevelEditor.cpp
+++ b/Engine/Source/Editor/LevelEditor/Private/LevelEditor.cpp
@@ -33,6 +33,7 @@
 #include "LevelViewportLayoutEntity.h"
 #include "PixelInspectorModule.h"
 #include "CommonMenuExtensionsModule.h"
+#include "ProjectDescriptor.h"
 
 // @todo Editor: remove this circular dependency
 #include "Interfaces/IMainFrameModule.h"
@@ -98,7 +99,12 @@ public:
 
 		// Create the tooltip showing more detailed information
 		FFormatNamedArguments TooltipArgs;
-		TooltipArgs.Add(TEXT("Version"), FText::FromString(EngineVersionString));
+		FString TooltipVersionStr = EngineVersionString;
+		if (IProjectManager::Get().GetCurrentProject() && IProjectManager::Get().GetCurrentProject()->bIsEnterpriseProject)
+		{
+			TooltipVersionStr += TEXT(" Enterprise");
+		}
+		TooltipArgs.Add(TEXT("Version"), FText::FromString(TooltipVersionStr));
 		TooltipArgs.Add(TEXT("Branch"), FText::FromString(FApp::GetBranchName()));
 		TooltipArgs.Add(TEXT("BuildConfiguration"), EBuildConfigurations::ToText(BuildConfig));
 		TooltipArgs.Add(TEXT("BuildDate"), FText::FromString(FApp::GetBuildDate()));
diff --git a/Engine/Source/Editor/MainFrame/Private/MainFrameModule.cpp b/Engine/Source/Editor/MainFrame/Private/MainFrameModule.cpp
--- a/Engine/Source/Editor/MainFrame/Private/MainFrameModule.cpp
+++ b/Engine/Source/Editor/MainFrame/Private/MainFrameModule.cpp
@@ -22,6 +22,7 @@
 #include "Developer/HotReload/Public/IHotReload.h"
 #include "Framework/Notifications/NotificationManager.h"
 #include "Widgets/Notifications/SNotificationList.h"
+#include "Widgets/Input/SEditableTextBox.h"
 #include "Framework/Commands/GenericCommands.h"
 #include "AnalyticsEventAttribute.h"
 #include "Interfaces/IAnalyticsProvider.h"
diff --git a/Engine/Source/Editor/PropertyEditor/Private/UserInterface/PropertyEditor/SPropertyEditorText.cpp b/Engine/Source/Editor/PropertyEditor/Private/UserInterface/PropertyEditor/SPropertyEditorText.cpp
--- a/Engine/Source/Editor/PropertyEditor/Private/UserInterface/PropertyEditor/SPropertyEditorText.cpp
+++ b/Engine/Source/Editor/PropertyEditor/Private/UserInterface/PropertyEditor/SPropertyEditorText.cpp
@@ -36,7 +36,7 @@ void SPropertyEditorText::Construct( const FArguments& InArgs, const TSharedRef<
 				.IsReadOnly(this, &SPropertyEditorText::IsReadOnly)
 				.AutoWrapText(true)
 				.ModiferKeyForNewLine(EModifierKey::Shift)
-				.IsPassword( bIsPassword )
+				//.IsPassword( bIsPassword )
 			]
 		];
 
diff --git a/Engine/Source/Editor/UnrealEd/Private/Toolkits/GlobalEditorCommonCommands.cpp b/Engine/Source/Editor/UnrealEd/Private/Toolkits/GlobalEditorCommonCommands.cpp
--- a/Engine/Source/Editor/UnrealEd/Private/Toolkits/GlobalEditorCommonCommands.cpp
+++ b/Engine/Source/Editor/UnrealEd/Private/Toolkits/GlobalEditorCommonCommands.cpp
@@ -135,6 +135,7 @@ void FGlobalEditorCommonCommands::OnSummonedConsoleCommandBox()
 
 		FDebugConsoleDelegates Delegates;
 		Delegates.OnConsoleCommandExecuted = FSimpleDelegate::CreateStatic( &CloseDebugConsole );
+		Delegates.OnCloseConsole = FSimpleDelegate::CreateStatic( &CloseDebugConsole );
 
 		OutputLogModule.ToggleDebugConsoleForWindow(WindowRef, EDebugConsoleStyle::Compact, Delegates);
 	}
diff --git a/Engine/Source/Programs/AutomationTool/Scripts/ListThirdPartySoftware.Automation.cs b/Engine/Source/Programs/AutomationTool/Scripts/ListThirdPartySoftware.Automation.cs
--- a/Engine/Source/Programs/AutomationTool/Scripts/ListThirdPartySoftware.Automation.cs
+++ b/Engine/Source/Programs/AutomationTool/Scripts/ListThirdPartySoftware.Automation.cs
@@ -17,6 +17,12 @@ class ListThirdPartySoftware : BuildCommand
 	{
 		CommandUtils.Log("************************* List Third Party Software");
 
+		string ProjectPath = ParseParamValue("Project", String.Empty);
+
+		//Add quotes to avoid issues with spaces in project path
+		if (ProjectPath != String.Empty)
+			ProjectPath = "\"" + ProjectPath + "\"";
+
 		// Parse the list of targets to list TPS for. Each target is specified by -Target="Name|Configuration|Platform" on the command line.
 		HashSet<FileReference> TpsFiles = new HashSet<FileReference>();
 		foreach(string Target in ParseParamValues(Params, "Target"))
@@ -24,8 +30,11 @@ class ListThirdPartySoftware : BuildCommand
 			// Get the path to store the exported JSON target data
 			FileReference OutputFile = FileReference.Combine(CommandUtils.EngineDirectory, "Intermediate", "Build", "ThirdParty.json");
 
-			IProcessResult Result = Run(UE4Build.GetUBTExecutable(), String.Format("{0} -jsonexport=\"{1}\" -skipbuild", Target.Replace('|', ' '), OutputFile.FullName), Options: ERunOptions.Default);
-			if(Result.ExitCode != 0)
+			IProcessResult Result;
+
+			Result = Run(UE4Build.GetUBTExecutable(), String.Format("{0} {1} -jsonexport=\"{2}\" -skipbuild", Target.Replace('|', ' '),  ProjectPath, OutputFile.FullName), Options: ERunOptions.Default);
+
+			if (Result.ExitCode != 0)
 			{
 				throw new AutomationException("Failed to run UBT");
 			}
diff --git a/Engine/Source/Programs/UnrealBuildTool/Configuration/UEBuildTarget.cs b/Engine/Source/Programs/UnrealBuildTool/Configuration/UEBuildTarget.cs
--- a/Engine/Source/Programs/UnrealBuildTool/Configuration/UEBuildTarget.cs
+++ b/Engine/Source/Programs/UnrealBuildTool/Configuration/UEBuildTarget.cs
@@ -3633,7 +3633,7 @@ namespace UnrealBuildTool
 			{
 				foreach(PluginInfo Plugin in NameToInfo.Values)
 				{
-					if(!NameToInstance.ContainsKey(Plugin.Name))
+					if(!NameToInstance.ContainsKey(Plugin.Name) && !Plugin.Descriptor.bCanBeUsedWithUnrealHeaderTool)
 					{
 						PluginReferenceDescriptor Reference = new PluginReferenceDescriptor(Plugin.Name, null, true);
 						AddPlugin(Reference, TargetRulesFile.GetFileName(), ExcludeFolders, NameToInstance, NameToInfo);
diff --git a/Engine/Source/Programs/UnrealBuildTool/ProjectFiles/VisualStudio/VCProject.cs b/Engine/Source/Programs/UnrealBuildTool/ProjectFiles/VisualStudio/VCProject.cs
--- a/Engine/Source/Programs/UnrealBuildTool/ProjectFiles/VisualStudio/VCProject.cs
+++ b/Engine/Source/Programs/UnrealBuildTool/ProjectFiles/VisualStudio/VCProject.cs
@@ -1299,10 +1299,6 @@ namespace UnrealBuildTool
 								DebugOptions += UProjectPath;
 								DebugOptions += " -skipcompile";
 							}
-							else if (TargetRulesObject.Type == TargetType.Editor && ProjectName == ProjectFileGenerator.EnterpriseProjectFileNameBase)
-							{
-								DebugOptions += " -enterprise";
-							}
 							else if (TargetRulesObject.Type == TargetType.Editor && ProjectName != "UE4")
 							{
 								DebugOptions += ProjectName;
diff --git a/Engine/Source/Runtime/Core/Private/HAL/ConsoleManager.cpp b/Engine/Source/Runtime/Core/Private/HAL/ConsoleManager.cpp
--- a/Engine/Source/Runtime/Core/Private/HAL/ConsoleManager.cpp
+++ b/Engine/Source/Runtime/Core/Private/HAL/ConsoleManager.cpp
@@ -994,33 +994,44 @@ void FConsoleManager::UnregisterConsoleObject(const TCHAR* Name, bool bKeepState
 
 void FConsoleManager::LoadHistoryIfNeeded()
 {
-	if(bHistoryWasLoaded)
+	if (bHistoryWasLoaded)
 	{
 		return;
 	}
 
 	bHistoryWasLoaded = true;
-
-	HistoryEntries.Empty();
+	HistoryEntriesMap.Reset();
 
 	FConfigFile Ini;
 
-	FString ConfigPath = FPaths::GeneratedConfigDir() + TEXT("ConsoleHistory.ini");
+	const FString ConfigPath = FPaths::GeneratedConfigDir() + TEXT("ConsoleHistory.ini");
 	ProcessIniContents(*ConfigPath, *ConfigPath, &Ini, false, false);
 
-	const FString History = TEXT("History");
-
-	FConfigSection* Section = Ini.Find(TEXT("ConsoleHistory"));
+	const FString SectionName = TEXT("ConsoleHistory");
+	const FName KeyName = TEXT("History");
 
-	if(Section)
+	for (const auto& ConfigPair : Ini)
 	{
-		for (auto It : *Section)
+		FString HistoryKey;
+		if (ConfigPair.Key == SectionName)
 		{
-			const FString& Key = It.Key.ToString();
+			// uses empty HistoryKey
+		}
+		else if (ConfigPair.Key.StartsWith(SectionName))
+		{
+			HistoryKey = ConfigPair.Key.Mid(SectionName.Len());
+		}
+		else
+		{
+			continue;
+		}
 
-			if(Key == History)
+		TArray<FString>& HistoryEntries = HistoryEntriesMap.FindOrAdd(HistoryKey);
+		for (const auto& ConfigSectionPair : ConfigPair.Value)
+		{
+			if (ConfigSectionPair.Key == KeyName)
 			{
-				HistoryEntries.Add(It.Value.GetValue());
+				HistoryEntries.Add(ConfigSectionPair.Value.GetValue());
 			}
 		}
 	}
@@ -1028,24 +1039,26 @@ void FConsoleManager::LoadHistoryIfNeeded()
 
 void FConsoleManager::SaveHistory()
 {
-	const FName History = TEXT("History");
-
 	FConfigFile Ini;
+	
+	const FString SectionName = TEXT("ConsoleHistory");
+	const FName KeyName = TEXT("History");
 
-	FString ConfigPath = FPaths::GeneratedConfigDir() + TEXT("ConsoleHistory.ini");
-
-	FConfigSection& Section = Ini.Add(TEXT("ConsoleHistory"));
-
-	for(auto It : HistoryEntries)
+	for (const auto& HistoryPair : HistoryEntriesMap)
 	{
-		Section.Add(History, It);
+		FConfigSection& Section = Ini.Add(FString::Printf(TEXT("%s%s"), *SectionName, *HistoryPair.Key));
+		for (const auto& HistoryEntry : HistoryPair.Value)
+		{
+			Section.Add(KeyName, HistoryEntry);
+		}
 	}
 
+	const FString ConfigPath = FPaths::GeneratedConfigDir() + TEXT("ConsoleHistory.ini");
+
 	Ini.Dirty = true;
 	Ini.Write(ConfigPath);
 }
 
-
 void FConsoleManager::ForEachConsoleObjectThatStartsWith(const FConsoleObjectVisitor& Visitor, const TCHAR* ThatStartsWith) const
 {
 	check(Visitor.IsBound());
@@ -1454,10 +1467,12 @@ void IConsoleManager::SetupSingleton()
 	check(Singleton);
 }
 
-void FConsoleManager::AddConsoleHistoryEntry(const TCHAR* Input)
+void FConsoleManager::AddConsoleHistoryEntry(const TCHAR* Key, const TCHAR* Input)
 {
 	LoadHistoryIfNeeded();
 
+	TArray<FString>& HistoryEntries = HistoryEntriesMap.FindOrAdd(Key);
+
 	// limit size to avoid a ever growing file
 	while(HistoryEntries.Num() > 64)
 	{
@@ -1471,11 +1486,11 @@ void FConsoleManager::AddConsoleHistoryEntry(const TCHAR* Input)
 	SaveHistory();
 }
 
-void FConsoleManager::GetConsoleHistory(TArray<FString>& Out)
+void FConsoleManager::GetConsoleHistory(const TCHAR* Key, TArray<FString>& Out)
 {
 	LoadHistoryIfNeeded();
 
-	Out = HistoryEntries;
+	Out = HistoryEntriesMap.FindOrAdd(Key);
 }
 
 bool FConsoleManager::IsNameRegistered(const TCHAR* Name) const
diff --git a/Engine/Source/Runtime/Core/Private/Misc/ConfigCacheIni.cpp b/Engine/Source/Runtime/Core/Private/Misc/ConfigCacheIni.cpp
--- a/Engine/Source/Runtime/Core/Private/Misc/ConfigCacheIni.cpp
+++ b/Engine/Source/Runtime/Core/Private/Misc/ConfigCacheIni.cpp
@@ -672,47 +672,8 @@ void FConfigFile::ProcessInputFileContents(const FString& Contents)
 				// If this line is delimited by quotes
 				if( *Value=='\"' )
 				{
-					FString PreprocessedValue = FString(Value).TrimQuotes().ReplaceQuotesWithEscapedQuotes();
-					const TCHAR* NewValue = *PreprocessedValue;
-
 					FString ProcessedValue;
-					//epic moelfke: fixed handling of escaped characters in quoted string
-					while (*NewValue && *NewValue != '\"')
-					{
-						if (*NewValue != '\\') // unescaped character
-						{
-							ProcessedValue += *NewValue++;
-						}
-						else if( *++NewValue == '\0')// escape character encountered at end
-						{
-							break;
-						}
-						else if (*NewValue == '\\') // escaped backslash "\\"
-						{
-							ProcessedValue += '\\';
-							NewValue++;
-						}
-						else if (*NewValue == '\"') // escaped double quote "\""
-						{
-							ProcessedValue += '\"';
-							NewValue++;
-						}
-						else if ( *NewValue == TEXT('n') )
-						{
-							ProcessedValue += TEXT('\n');
-							NewValue++;
-						}
-						else if( *NewValue == TEXT('u') && NewValue[1] && NewValue[2] && NewValue[3] && NewValue[4] )	// \uXXXX - UNICODE code point
-						{
-							ProcessedValue += (TCHAR)(FParse::HexDigit(NewValue[1])*(1<<12) + FParse::HexDigit(NewValue[2])*(1<<8) + FParse::HexDigit(NewValue[3])*(1<<4) + FParse::HexDigit(NewValue[4]));
-							NewValue += 5;
-						}
-						else if( NewValue[1] ) // some other escape sequence, assume it's a hex character value
-						{
-							ProcessedValue += (TCHAR)(FParse::HexDigit(NewValue[0])*16 + FParse::HexDigit(NewValue[1]));
-							NewValue += 2;
-						}
-					}
+					FParse::QuotedString(Value, ProcessedValue);
 
 					// Add this pair to the current FConfigSection
 					CurrentSection->Add(Start, *ProcessedValue);
@@ -800,6 +761,12 @@ bool FConfigFile::ShouldExportQuotedString(const FString& PropertyValue)
 		{
 			return true;
 		}
+
+		// ... it contains an unescaped new-line
+		if (!bEscapeNextChar && (NextChar == TEXT('\r') || NextChar == TEXT('\n')))
+		{
+			return true;
+		}
 	}
 
 	return false;
diff --git a/Engine/Source/Runtime/Core/Private/Misc/Parse.cpp b/Engine/Source/Runtime/Core/Private/Misc/Parse.cpp
--- a/Engine/Source/Runtime/Core/Private/Misc/Parse.cpp
+++ b/Engine/Source/Runtime/Core/Private/Misc/Parse.cpp
@@ -358,7 +358,12 @@ bool FParse::QuotedString( const TCHAR* Buffer, FString& Value, int32* OutNumCha
 		return false;
 	}
 
-	while (*Buffer && *Buffer != TCHAR('"') && *Buffer != TCHAR('\n') && *Buffer != TCHAR('\r'))
+	auto ShouldParse = [](const TCHAR Ch)
+	{
+		return Ch != 0 && Ch != TCHAR('"') && Ch != TCHAR('\n') && Ch != TCHAR('\r');
+	};
+
+	while (ShouldParse(*Buffer))
 	{
 		if (*Buffer != TCHAR('\\')) // unescaped character
 		{
@@ -369,7 +374,7 @@ bool FParse::QuotedString( const TCHAR* Buffer, FString& Value, int32* OutNumCha
 			Value += TEXT("\\");
 			++Buffer;
 		}
-		else if (*Buffer == TCHAR('\"')) // escaped double quote "\""
+		else if (*Buffer == TCHAR('"')) // escaped double quote "\""
 		{
 			Value += TCHAR('"');
 			++Buffer;
@@ -394,10 +399,68 @@ bool FParse::QuotedString( const TCHAR* Buffer, FString& Value, int32* OutNumCha
 			Value += TCHAR('\t');
 			++Buffer;
 		}
-		else // some other escape sequence, assume it's a hex character value
+		else if (FChar::IsOctDigit(*Buffer)) // octal sequence (\012)
+		{
+			FString OctSequence;
+			while (ShouldParse(*Buffer) && FChar::IsOctDigit(*Buffer) && OctSequence.Len() < 3) // Octal sequences can only be up-to 3 digits long
+			{
+				OctSequence += *Buffer++;
+			}
+
+			Value += (TCHAR)FCString::Strtoi(*OctSequence, nullptr, 8);
+		}
+		else if (*Buffer == TCHAR('x')) // hex sequence (\xBEEF)
+		{
+			++Buffer;
+
+			FString HexSequence;
+			while (ShouldParse(*Buffer) && FChar::IsHexDigit(*Buffer))
+			{
+				HexSequence += *Buffer++;
+			}
+
+			Value += (TCHAR)FCString::Strtoi(*HexSequence, nullptr, 16);
+		}
+		else if (*Buffer == TCHAR('u')) // UTF-16 sequence (\u1234)
+		{
+			++Buffer;
+
+			FString UnicodeSequence;
+			while (ShouldParse(*Buffer) && FChar::IsHexDigit(*Buffer) && UnicodeSequence.Len() < 4) // UTF-16 sequences can only be up-to 4 digits long
+			{
+				UnicodeSequence += *Buffer++;
+			}
+
+			const uint32 UnicodeCodepoint = (uint32)FCString::Strtoi(*UnicodeSequence, nullptr, 16);
+
+			FString UnicodeString;
+			if (FUnicodeChar::CodepointToString(UnicodeCodepoint, UnicodeString))
+			{
+				Value += MoveTemp(UnicodeString);
+			}
+		}
+		else if (*Buffer == TCHAR('U')) // UTF-32 sequence (\U12345678)
+		{
+			++Buffer;
+
+			FString UnicodeSequence;
+			while (ShouldParse(*Buffer) && FChar::IsHexDigit(*Buffer) && UnicodeSequence.Len() < 8) // UTF-32 sequences can only be up-to 8 digits long
+			{
+				UnicodeSequence += *Buffer++;
+			}
+
+			const uint32 UnicodeCodepoint = (uint32)FCString::Strtoi(*UnicodeSequence, nullptr, 16);
+
+			FString UnicodeString;
+			if (FUnicodeChar::CodepointToString(UnicodeCodepoint, UnicodeString))
+			{
+				Value += MoveTemp(UnicodeString);
+			}
+		}
+		else // unhandled escape sequence
 		{
-			Value += FString::Printf(TEXT("%c"), (HexDigit(Buffer[0]) * 16) + HexDigit(Buffer[1]));
-			Buffer += 2;
+			Value += TEXT("\\");
+			Value += *Buffer++;
 		}
 	}
 
diff --git a/Engine/Source/Runtime/Core/Private/Misc/Paths.cpp b/Engine/Source/Runtime/Core/Private/Misc/Paths.cpp
--- a/Engine/Source/Runtime/Core/Private/Misc/Paths.cpp
+++ b/Engine/Source/Runtime/Core/Private/Misc/Paths.cpp
@@ -160,6 +160,11 @@ FString FPaths::EnterprisePluginsDir()
 	return EnterpriseDir() + TEXT("Plugins/");
 }
 
+FString FPaths::EnterpriseFeaturePackDir()
+{
+	return FPaths::EnterpriseDir() + TEXT("FeaturePacks/");
+}
+
 FString FPaths::RootDir()
 {
 	return FString(FPlatformMisc::RootDir());
diff --git a/Engine/Source/Runtime/Core/Public/HAL/ConsoleManager.h b/Engine/Source/Runtime/Core/Public/HAL/ConsoleManager.h
--- a/Engine/Source/Runtime/Core/Public/HAL/ConsoleManager.h
+++ b/Engine/Source/Runtime/Core/Public/HAL/ConsoleManager.h
@@ -71,8 +71,8 @@ public:
 	virtual void ForEachConsoleObjectThatStartsWith(const FConsoleObjectVisitor& Visitor, const TCHAR* ThatStartsWith) const override;
 	virtual void ForEachConsoleObjectThatContains(const FConsoleObjectVisitor& Visitor, const TCHAR* ThatContains) const override;
 	virtual bool ProcessUserConsoleInput(const TCHAR* InInput, FOutputDevice& Ar, UWorld* InWorld) override;
-	virtual void AddConsoleHistoryEntry(const TCHAR* Input) override;
-	virtual void GetConsoleHistory(TArray<FString>& Out) override;
+	virtual void AddConsoleHistoryEntry(const TCHAR* Key, const TCHAR* Input) override;
+	virtual void GetConsoleHistory(const TCHAR* Key, TArray<FString>& Out) override;
 	virtual bool IsNameRegistered(const TCHAR* Name) const override;	
 	virtual void RegisterThreadPropagation(uint32 ThreadId, IConsoleThreadPropagation* InCallback) override;
 	virtual void UnregisterConsoleObject( IConsoleObject* Object, bool bKeepState) override;
@@ -83,8 +83,8 @@ private: // ----------------------------------------------------
 	// [name] = pointer (pointer must not be 0)
 	TMap<FString, IConsoleObject*> ConsoleObjects;
 
-	TArray<FString>	HistoryEntries;
 	bool bHistoryWasLoaded;
+	TMap<FString, TArray<FString>>	HistoryEntriesMap;
 	TArray<FConsoleCommandDelegate>	ConsoleVariableChangeSinks;
 
 	IConsoleThreadPropagation* ThreadPropagationCallback;
@@ -134,9 +134,9 @@ private: // ----------------------------------------------------
 	 */
 	void UnregisterConsoleObject(const TCHAR* Name, bool bKeepState);
 
-	// clears HistoryEntries and reads it from the .ini file
+	// reads HistoryEntriesMap from the .ini file (if not already loaded)
 	void LoadHistoryIfNeeded();
 
-	// clears HistoryEntries and writes it the .ini file
+	// writes HistoryEntriesMap to the .ini file
 	void SaveHistory();
 };
diff --git a/Engine/Source/Runtime/Core/Public/HAL/IConsoleManager.h b/Engine/Source/Runtime/Core/Public/HAL/IConsoleManager.h
--- a/Engine/Source/Runtime/Core/Public/HAL/IConsoleManager.h
+++ b/Engine/Source/Runtime/Core/Public/HAL/IConsoleManager.h
@@ -9,6 +9,7 @@
 #include "Logging/LogMacros.h"
 #include "Delegates/IDelegateInstance.h"
 #include "Delegates/Delegate.h"
+#include "Features/IModularFeature.h"
 
 #define TRACK_CONSOLE_FIND_COUNT !(UE_BUILD_SHIPPING || UE_BUILD_TEST)
 
@@ -367,6 +368,70 @@ private:
 };
 
 
+/**
+ * Handles executing console commands
+ */
+class IConsoleCommandExecutor : public IModularFeature
+{
+public:
+	/**
+	 * Get the name identifying this modular feature set.
+	 */
+	static FName ModularFeatureName()
+	{
+		static const FName Name = TEXT("ConsoleCommandExecutor");
+		return Name;
+	}
+
+	/**
+	 * Get the name of this executor.
+	 */
+	virtual FName GetName() const = 0;
+
+	/**
+	 * Get the display name of this executor.
+	 */
+	virtual FText GetDisplayName() const = 0;
+
+	/**
+	 * Get the description of this executor.
+	 */
+	virtual FText GetDescription() const = 0;
+
+	/**
+	 * Get the hint text of this executor.
+	 */
+	virtual FText GetHintText() const = 0;
+
+	/**
+	 * Get the list of auto-complete suggestions for the given command.
+	 */
+	virtual void GetAutoCompleteSuggestions(const TCHAR* Input, TArray<FString>& Out) = 0;
+
+	/**
+	 * Get the list of commands that this executor has recently processed.
+	 */
+	virtual void GetExecHistory(TArray<FString>& Out) = 0;
+
+	/**
+	 * Execute the given command using this executor.
+	 * @return true if the command was recognized.
+	 */
+	virtual bool Exec(const TCHAR* Input) = 0;
+
+	/**
+	 * True if we allow the console to be closed using the "open console" hot-key.
+	 * @note Some scripting languages use the default "open console" hot-key (~) in their code, so these should return false.
+	 */
+	virtual bool AllowHotKeyClose() const = 0;
+
+	/**
+	 * True if we allow the console to create multi-line commands.
+	 */
+	virtual bool AllowMultiLine() const = 0;
+};
+
+
 /**
  * handles console commands and variables, registered console variables are released on destruction
  */
@@ -598,11 +663,11 @@ struct CORE_API IConsoleManager
 	/**
 	 * @param Input - must not be 0
 	 */
-	virtual void AddConsoleHistoryEntry(const TCHAR* Input) = 0;
+	virtual void AddConsoleHistoryEntry(const TCHAR* Key, const TCHAR* Input) = 0;
 	
 	/**
 	 */
-	virtual void GetConsoleHistory(TArray<FString>& Out) = 0; 
+	virtual void GetConsoleHistory(const TCHAR* Key, TArray<FString>& Out) = 0; 
 
 	/**
 	 * Check if a name (command or variable) has been registered with the console manager
diff --git a/Engine/Source/Runtime/Core/Public/Math/Vector4.h b/Engine/Source/Runtime/Core/Public/Math/Vector4.h
--- a/Engine/Source/Runtime/Core/Public/Math/Vector4.h
+++ b/Engine/Source/Runtime/Core/Public/Math/Vector4.h
@@ -123,6 +123,14 @@ public:
 	 */
 	FORCEINLINE FVector4 operator-(const FVector4& V) const;
 
+	/**
+	 * Subtracts another vector to this one.
+	 *
+	 * @param V The other vector to subtract.
+	 * @return Copy of the vector after subtraction.
+	 */
+	FORCEINLINE FVector4 operator-=(const FVector4& V);
+
 	/**
 	 * Gets the result of scaling this vector.
 	 *
@@ -527,6 +535,14 @@ FORCEINLINE FVector4 FVector4::operator-(const FVector4& V) const
 }
 
 
+FORCEINLINE FVector4 FVector4::operator-=(const FVector4& V)
+{
+	X -= V.X; Y -= V.Y; Z -= V.Z; W -= V.W;
+	DiagnosticCheckNaN();
+	return *this;
+}
+
+
 FORCEINLINE FVector4 FVector4::operator*(float Scale) const
 {
 	return FVector4(X * Scale, Y * Scale, Z * Scale, W * Scale);
diff --git a/Engine/Source/Runtime/Core/Public/Misc/Paths.h b/Engine/Source/Runtime/Core/Public/Misc/Paths.h
--- a/Engine/Source/Runtime/Core/Public/Misc/Paths.h
+++ b/Engine/Source/Runtime/Core/Public/Misc/Paths.h
@@ -96,6 +96,13 @@ public:
 	*/
 	static FString EnterprisePluginsDir();
 
+	/**
+	* Returns the enterprise FeaturePack directory
+	*
+	* @return FeaturePack directory.
+	*/
+	static FString EnterpriseFeaturePackDir();
+
 	/**
 	 * Returns the root directory of the engine directory tree
 	 *
diff --git a/Engine/Source/Runtime/CoreUObject/Private/UObject/Class.cpp b/Engine/Source/Runtime/CoreUObject/Private/UObject/Class.cpp
--- a/Engine/Source/Runtime/CoreUObject/Private/UObject/Class.cpp
+++ b/Engine/Source/Runtime/CoreUObject/Private/UObject/Class.cpp
@@ -4753,6 +4753,12 @@ UScriptStruct* TBaseStructure<FRotator>::Get()
 	return ScriptStruct;
 }
 
+UScriptStruct* TBaseStructure<FQuat>::Get()
+{
+	static auto ScriptStruct = StaticGetBaseStructureInternal(TEXT("Quat"));
+	return ScriptStruct;
+}
+
 UScriptStruct* TBaseStructure<FTransform>::Get()
 {
 	static auto ScriptStruct = StaticGetBaseStructureInternal(TEXT("Transform"));
@@ -4771,6 +4777,12 @@ UScriptStruct* TBaseStructure<FColor>::Get()
 	return ScriptStruct;
 }
 
+UScriptStruct* TBaseStructure<FPlane>::Get()
+{
+	static auto ScriptStruct = StaticGetBaseStructureInternal(TEXT("Plane"));
+	return ScriptStruct;
+}
+
 UScriptStruct* TBaseStructure<FVector>::Get()
 {
 	static auto ScriptStruct = StaticGetBaseStructureInternal(TEXT("Vector"));
@@ -4783,6 +4795,12 @@ UScriptStruct* TBaseStructure<FVector2D>::Get()
 	return ScriptStruct;
 }
 
+UScriptStruct* TBaseStructure<FVector4>::Get()
+{
+	static auto ScriptStruct = StaticGetBaseStructureInternal(TEXT("Vector4"));
+	return ScriptStruct;
+}
+
 UScriptStruct* TBaseStructure<FRandomStream>::Get()
 {
 	static auto ScriptStruct = StaticGetBaseStructureInternal(TEXT("RandomStream"));
diff --git a/Engine/Source/Runtime/CoreUObject/Public/UObject/Class.h b/Engine/Source/Runtime/CoreUObject/Public/UObject/Class.h
--- a/Engine/Source/Runtime/CoreUObject/Public/UObject/Class.h
+++ b/Engine/Source/Runtime/CoreUObject/Public/UObject/Class.h
@@ -3013,6 +3013,10 @@ struct FStructUtils
 
 template< class T > struct TBaseStructure
 {
+	static UScriptStruct* Get()
+	{
+		return T::StaticStruct();
+	}
 };
 
 template<> struct TBaseStructure<FRotator>
@@ -3020,6 +3024,11 @@ template<> struct TBaseStructure<FRotator>
 	COREUOBJECT_API static UScriptStruct* Get();
 };
 
+template<> struct TBaseStructure<FQuat>
+{
+	COREUOBJECT_API static UScriptStruct* Get();
+};
+
 template<> struct TBaseStructure<FTransform>
 {
 	COREUOBJECT_API static UScriptStruct* Get();
@@ -3035,6 +3044,11 @@ template<> struct TBaseStructure<FColor>
 	COREUOBJECT_API static UScriptStruct* Get();
 };
 
+template<> struct  TBaseStructure<FPlane>
+{
+	COREUOBJECT_API static UScriptStruct* Get();
+};
+
 template<> struct  TBaseStructure<FVector>
 {
 	COREUOBJECT_API static UScriptStruct* Get();
@@ -3045,6 +3059,11 @@ template<> struct TBaseStructure<FVector2D>
 	COREUOBJECT_API static UScriptStruct* Get();
 };
 
+template<> struct TBaseStructure<FVector4>
+{
+	COREUOBJECT_API static UScriptStruct* Get();
+};
+
 template<> struct TBaseStructure<FRandomStream>
 {
 	COREUOBJECT_API static UScriptStruct* Get();
diff --git a/Engine/Source/Runtime/Engine/Classes/Kismet/DataTableFunctionLibrary.h b/Engine/Source/Runtime/Engine/Classes/Kismet/DataTableFunctionLibrary.h
--- a/Engine/Source/Runtime/Engine/Classes/Kismet/DataTableFunctionLibrary.h
+++ b/Engine/Source/Runtime/Engine/Classes/Kismet/DataTableFunctionLibrary.h
@@ -38,6 +38,10 @@ class ENGINE_API UDataTableFunctionLibrary : public UBlueprintFunctionLibrary
 	UFUNCTION(BlueprintCallable, Category = "DataTable")
 	static void GetDataTableRowNames(UDataTable* Table, TArray<FName>& OutRowNames);
 
+	/** Export from the DataTable all the row for one column. Export it as string. The row name is not included. */
+	UFUNCTION(BlueprintCallable, Category = "DataTable")
+	static TArray<FString> GetDataTableColumnAsString(const UDataTable* DataTable, FName PropertyName);
+
     /** Get a Row from a DataTable given a RowName */
     UFUNCTION(BlueprintCallable, CustomThunk, Category = "DataTable", meta=(CustomStructureParam = "OutRow", BlueprintInternalUseOnly="true"))
     static bool GetDataTableRowFromName(UDataTable* Table, FName RowName, FTableRowBase& OutRow);
diff --git a/Engine/Source/Runtime/Engine/Private/ActorConstruction.cpp b/Engine/Source/Runtime/Engine/Private/ActorConstruction.cpp
--- a/Engine/Source/Runtime/Engine/Private/ActorConstruction.cpp
+++ b/Engine/Source/Runtime/Engine/Private/ActorConstruction.cpp
@@ -1082,10 +1082,16 @@ void AActor::CheckComponentInstanceName(const FName InName)
 			// Try and pick a good name
 			FString ConflictingObjectName = ConflictingObject->GetName();
 			int32 CharIndex = ConflictingObjectName.Len() - 1;
-			while (FChar::IsDigit(ConflictingObjectName[CharIndex]))
+			while (CharIndex >= 0 && FChar::IsDigit(ConflictingObjectName[CharIndex]))
 			{
 				--CharIndex;
 			}
+			// Name is only composed of digits not a name conflict resolution
+			if (CharIndex < 0)
+			{
+				return;
+			}
+
 			int32 Counter = 0;
 			if (CharIndex < ConflictingObjectName.Len() - 1)
 			{
diff --git a/Engine/Source/Runtime/Engine/Private/DataTableFunctionLibrary.cpp b/Engine/Source/Runtime/Engine/Private/DataTableFunctionLibrary.cpp
--- a/Engine/Source/Runtime/Engine/Private/DataTableFunctionLibrary.cpp
+++ b/Engine/Source/Runtime/Engine/Private/DataTableFunctionLibrary.cpp
@@ -26,6 +26,16 @@ void UDataTableFunctionLibrary::EvaluateCurveTableRow(UCurveTable* CurveTable, F
 	}
 }
 
+TArray<FString> UDataTableFunctionLibrary::GetDataTableColumnAsString(const UDataTable* DataTable, FName PropertyName)
+{
+	if (DataTable && PropertyName != NAME_None)
+	{
+		EDataTableExportFlags ExportFlags = EDataTableExportFlags::None;
+		return DataTableUtils::GetColumnDataAsString(DataTable, PropertyName, ExportFlags);
+	}
+	return TArray<FString>();
+}
+
 bool UDataTableFunctionLibrary::Generic_GetDataTableRowFromName(UDataTable* Table, FName RowName, void* OutRowPtr)
 {
 	bool bFoundRow = false;
diff --git a/Engine/Source/Runtime/Engine/Private/DataTableUtils.cpp b/Engine/Source/Runtime/Engine/Private/DataTableUtils.cpp
--- a/Engine/Source/Runtime/Engine/Private/DataTableUtils.cpp
+++ b/Engine/Source/Runtime/Engine/Private/DataTableUtils.cpp
@@ -3,6 +3,7 @@
 #include "DataTableUtils.h"
 #include "UObject/UnrealType.h"
 #include "UObject/TextProperty.h"
+#include "Engine/DataTable.h"
 #include "Engine/UserDefinedEnum.h"
 #include "Engine/UserDefinedStruct.h"
 #include "Policies/PrettyJsonPrintPolicy.h"
@@ -478,3 +479,28 @@ FString DataTableUtils::GetPropertyDisplayName(const UProperty* Prop, const FStr
 	return DefaultName;
 #endif // WITH_EDITOR
 }
+
+TArray<FString> DataTableUtils::GetColumnDataAsString(const UDataTable* InTable, const FName& PropertyName, const EDataTableExportFlags InDTExportFlags)
+{
+	TArray<FString> Result;
+	if (!ensure(InTable))
+	{
+		return Result;
+	}
+	if (!ensure(PropertyName != NAME_None))
+	{
+		return Result;
+	}
+
+	UProperty* ColumnProperty = InTable->FindTableProperty(PropertyName);
+	if (ColumnProperty)
+	{
+		for (auto RowIt = InTable->RowMap.CreateConstIterator(); RowIt; ++RowIt)
+		{
+			uint8* RowData = RowIt.Value();
+			Result.Add(GetPropertyValueAsString(ColumnProperty, RowData, InDTExportFlags));
+		}
+	}
+
+	return Result;
+}
diff --git a/Engine/Source/Runtime/Engine/Public/DataTableUtils.h b/Engine/Source/Runtime/Engine/Public/DataTableUtils.h
--- a/Engine/Source/Runtime/Engine/Public/DataTableUtils.h
+++ b/Engine/Source/Runtime/Engine/Public/DataTableUtils.h
@@ -5,6 +5,8 @@
 #include "CoreMinimal.h"
 #include "Misc/EnumClassFlags.h"
 
+class UDataTable;
+
 ENGINE_API DECLARE_LOG_CATEGORY_EXTERN(LogDataTable, Log, All);
 
 enum class EDataTableExportFlags : uint8
@@ -90,4 +92,9 @@ namespace DataTableUtils
 	 * Util to get the friendly display name of a given property.
 	 */
 	ENGINE_API FString GetPropertyDisplayName(const UProperty* Prop, const FString& DefaultName);
+
+	/**
+	 * Output each row for a specific column/property in the table (doesn't include the title)
+	 */
+	ENGINE_API TArray<FString> GetColumnDataAsString(const UDataTable* InTable, const FName& PropertyName, const EDataTableExportFlags InDTExportFlags);
 }
diff --git a/Engine/Source/Runtime/Slate/Private/Widgets/Input/SEditableText.cpp b/Engine/Source/Runtime/Slate/Private/Widgets/Input/SEditableText.cpp
--- a/Engine/Source/Runtime/Slate/Private/Widgets/Input/SEditableText.cpp
+++ b/Engine/Source/Runtime/Slate/Private/Widgets/Input/SEditableText.cpp
@@ -34,6 +34,7 @@ void SEditableText::Construct( const FArguments& InArgs )
 	VirtualKeyboardType = InArgs._VirtualKeyboardType;
 	VirtualKeyboardTrigger = InArgs._VirtualKeyboardTrigger;
 	VirtualKeyboardDismissAction = InArgs._VirtualKeyboardDismissAction;
+	OnKeyCharHandler = InArgs._OnKeyCharHandler;
 	OnKeyDownHandler = InArgs._OnKeyDownHandler;
 
 	Font = InArgs._Font;
@@ -171,7 +172,20 @@ void SEditableText::OnFocusLost( const FFocusEvent& InFocusEvent )
 
 FReply SEditableText::OnKeyChar( const FGeometry& MyGeometry, const FCharacterEvent& InCharacterEvent )
 {
-	return EditableTextLayout->HandleKeyChar(InCharacterEvent);
+	FReply Reply = FReply::Unhandled();
+
+	// First call the user defined key handler, there might be overrides to normal functionality
+	if (OnKeyCharHandler.IsBound())
+	{
+		Reply = OnKeyCharHandler.Execute(MyGeometry, InCharacterEvent);
+	}
+
+	if (!Reply.IsEventHandled())
+	{
+		Reply = EditableTextLayout->HandleKeyChar(InCharacterEvent);
+	}
+
+	return Reply;
 }
 
 FReply SEditableText::OnKeyDown( const FGeometry& MyGeometry, const FKeyEvent& InKeyEvent )
@@ -484,7 +498,12 @@ bool SEditableText::CanInsertCarriageReturn() const
 
 bool SEditableText::CanTypeCharacter(const TCHAR InChar) const
 {
-	return !OnIsTypedCharValid.IsBound() || OnIsTypedCharValid.Execute(InChar);
+	if (OnIsTypedCharValid.IsBound())
+	{
+		return OnIsTypedCharValid.Execute(InChar);
+	}
+
+	return InChar != TEXT('\t');
 }
 
 void SEditableText::EnsureActiveTick()
diff --git a/Engine/Source/Runtime/Slate/Private/Widgets/Input/SEditableTextBox.cpp b/Engine/Source/Runtime/Slate/Private/Widgets/Input/SEditableTextBox.cpp
--- a/Engine/Source/Runtime/Slate/Private/Widgets/Input/SEditableTextBox.cpp
+++ b/Engine/Source/Runtime/Slate/Private/Widgets/Input/SEditableTextBox.cpp
@@ -56,6 +56,7 @@ void SEditableTextBox::Construct( const FArguments& InArgs )
 					.OnTextCommitted( InArgs._OnTextCommitted )
 					.MinDesiredWidth( InArgs._MinDesiredWidth )
 					.SelectAllTextOnCommit( InArgs._SelectAllTextOnCommit )
+					.OnKeyCharHandler( InArgs._OnKeyCharHandler )			
 					.OnKeyDownHandler( InArgs._OnKeyDownHandler )
 					.VirtualKeyboardType( InArgs._VirtualKeyboardType )
 					.VirtualKeyboardTrigger( InArgs._VirtualKeyboardTrigger )
@@ -132,6 +133,12 @@ void SEditableTextBox::SetError( const FString& InError )
 }
 
 
+void SEditableTextBox::SetOnKeyCharHandler(FOnKeyChar InOnKeyCharHandler)
+{
+	EditableText->SetOnKeyCharHandler(InOnKeyCharHandler);
+}
+
+
 void SEditableTextBox::SetOnKeyDownHandler(FOnKeyDown InOnKeyDownHandler)
 {
 	EditableText->SetOnKeyDownHandler(InOnKeyDownHandler);
diff --git a/Engine/Source/Runtime/Slate/Private/Widgets/Input/SMultiLineEditableTextBox.cpp b/Engine/Source/Runtime/Slate/Private/Widgets/Input/SMultiLineEditableTextBox.cpp
--- a/Engine/Source/Runtime/Slate/Private/Widgets/Input/SMultiLineEditableTextBox.cpp
+++ b/Engine/Source/Runtime/Slate/Private/Widgets/Input/SMultiLineEditableTextBox.cpp
@@ -99,7 +99,9 @@ void SMultiLineEditableTextBox::Construct( const FArguments& InArgs )
 					.Marshaller( InArgs._Marshaller )
 					.Font( this, &SMultiLineEditableTextBox::DetermineFont )
 					.IsReadOnly( InArgs._IsReadOnly )
+					.AllowMultiLine( InArgs._AllowMultiLine )
 					.OnContextMenuOpening( InArgs._OnContextMenuOpening )
+					.OnIsTypedCharValid( InArgs._OnIsTypedCharValid )
 					.OnTextChanged( InArgs._OnTextChanged )
 					.OnTextCommitted( InArgs._OnTextCommitted )
 					.OnCursorMoved( InArgs._OnCursorMoved )
@@ -119,7 +121,8 @@ void SMultiLineEditableTextBox::Construct( const FArguments& InArgs )
 					.VScrollBar(VScrollBar)
 					.OnHScrollBarUserScrolled(InArgs._OnHScrollBarUserScrolled)
 					.OnVScrollBarUserScrolled(InArgs._OnVScrollBarUserScrolled)
-					.OnKeyDownHandler( InArgs._OnKeyDownHandler)
+					.OnKeyCharHandler(InArgs._OnKeyCharHandler)
+					.OnKeyDownHandler(InArgs._OnKeyDownHandler)
 					.ModiferKeyForNewLine(InArgs._ModiferKeyForNewLine)
 					.VirtualKeyboardTrigger( InArgs._VirtualKeyboardTrigger )
 					.VirtualKeyboardDismissAction( InArgs._VirtualKeyboardDismissAction )
@@ -465,6 +468,11 @@ void SMultiLineEditableTextBox::Refresh()
 	return EditableText->Refresh();
 }
 
+void SMultiLineEditableTextBox::SetOnKeyCharHandler(FOnKeyChar InOnKeyCharHandler)
+{
+	EditableText->SetOnKeyCharHandler(InOnKeyCharHandler);
+}
+
 void SMultiLineEditableTextBox::SetOnKeyDownHandler(FOnKeyDown InOnKeyDownHandler)
 {
 	EditableText->SetOnKeyDownHandler(InOnKeyDownHandler);
diff --git a/Engine/Source/Runtime/Slate/Private/Widgets/Text/SMultiLineEditableText.cpp b/Engine/Source/Runtime/Slate/Private/Widgets/Text/SMultiLineEditableText.cpp
--- a/Engine/Source/Runtime/Slate/Private/Widgets/Text/SMultiLineEditableText.cpp
+++ b/Engine/Source/Runtime/Slate/Private/Widgets/Text/SMultiLineEditableText.cpp
@@ -29,9 +29,11 @@ void SMultiLineEditableText::Construct( const FArguments& InArgs )
 {
 	bIsReadOnly = InArgs._IsReadOnly;
 
+	OnIsTypedCharValid = InArgs._OnIsTypedCharValid;
 	OnTextChangedCallback = InArgs._OnTextChanged;
 	OnTextCommittedCallback = InArgs._OnTextCommitted;
 	OnCursorMovedCallback = InArgs._OnCursorMoved;
+	bAllowMultiLine = InArgs._AllowMultiLine;
 	bSelectAllTextWhenFocused = InArgs._SelectAllTextWhenFocused;
 	bClearTextSelectionOnFocusLoss = InArgs._ClearTextSelectionOnFocusLoss;
 	bClearKeyboardFocusOnCommit = InArgs._ClearKeyboardFocusOnCommit;
@@ -42,6 +44,7 @@ void SMultiLineEditableText::Construct( const FArguments& InArgs )
 	VirtualKeyboardDismissAction = InArgs._VirtualKeyboardDismissAction;
 	OnHScrollBarUserScrolled = InArgs._OnHScrollBarUserScrolled;
 	OnVScrollBarUserScrolled = InArgs._OnVScrollBarUserScrolled;
+	OnKeyCharHandler = InArgs._OnKeyCharHandler;
 	OnKeyDownHandler = InArgs._OnKeyDownHandler;
 	ModiferKeyForNewLine = InArgs._ModiferKeyForNewLine;
 
@@ -240,7 +243,7 @@ bool SMultiLineEditableText::IsTextPassword() const
 
 bool SMultiLineEditableText::IsMultiLineTextEdit() const
 {
-	return true;
+	return bAllowMultiLine.Get(true);
 }
 
 bool SMultiLineEditableText::ShouldJumpCursorToEndWhenFocused() const
@@ -280,7 +283,12 @@ bool SMultiLineEditableText::CanInsertCarriageReturn() const
 
 bool SMultiLineEditableText::CanTypeCharacter(const TCHAR InChar) const
 {
-	return true;
+	if (OnIsTypedCharValid.IsBound())
+	{
+		return OnIsTypedCharValid.Execute(InChar);
+	}
+
+	return InChar != TEXT('\t');
 }
 
 void SMultiLineEditableText::EnsureActiveTick()
@@ -562,7 +570,20 @@ bool SMultiLineEditableText::SupportsKeyboardFocus() const
 
 FReply SMultiLineEditableText::OnKeyChar( const FGeometry& MyGeometry,const FCharacterEvent& InCharacterEvent )
 {
-	return EditableTextLayout->HandleKeyChar(InCharacterEvent);
+	FReply Reply = FReply::Unhandled();
+
+	// First call the user defined key handler, there might be overrides to normal functionality
+	if (OnKeyCharHandler.IsBound())
+	{
+		Reply = OnKeyCharHandler.Execute(MyGeometry, InCharacterEvent);
+	}
+
+	if (!Reply.IsEventHandled())
+	{
+		Reply = EditableTextLayout->HandleKeyChar(InCharacterEvent);
+	}
+
+	return Reply;
 }
 
 FReply SMultiLineEditableText::OnKeyDown( const FGeometry& MyGeometry, const FKeyEvent& InKeyEvent )
diff --git a/Engine/Source/Runtime/Slate/Private/Widgets/Text/SlateEditableTextLayout.cpp b/Engine/Source/Runtime/Slate/Private/Widgets/Text/SlateEditableTextLayout.cpp
--- a/Engine/Source/Runtime/Slate/Private/Widgets/Text/SlateEditableTextLayout.cpp
+++ b/Engine/Source/Runtime/Slate/Private/Widgets/Text/SlateEditableTextLayout.cpp
@@ -809,9 +809,6 @@ FReply FSlateEditableTextLayout::HandleKeyChar(const FCharacterEvent& InCharacte
 		}
 		break;
 
-	case TCHAR('\t'):	// Tab
-		return FReply::Handled();
-
 	case TCHAR('\n'):	// Newline (Ctrl+Enter), we handle adding new lines via HandleCarriageReturn rather than processing newline characters
 		return FReply::Handled();
 
@@ -943,6 +940,10 @@ FReply FSlateEditableTextLayout::HandleKeyDown(const FKeyEvent& InKeyEvent)
 		FScopedEditableTextTransaction TextTransaction(*this);
 		Reply = BoolToReply(HandleDelete());
 	}
+	else if (Key == EKeys::Tab && OwnerWidget->CanTypeCharacter(TEXT('\t')))
+	{
+		Reply = FReply::Handled();
+	}
 	else if (Key == EKeys::Escape)
 	{
 		Reply = BoolToReply(HandleEscape());
@@ -3196,7 +3197,8 @@ void FSlateEditableTextLayout::Tick(const FGeometry& AllottedGeometry, const dou
 
 	// If we're auto-wrapping, we need to hide the scrollbars until the first valid auto-wrap has been performed
 	// If we don't do this, then we can get some nasty layout shuffling as the scrollbars appear for one frame and then vanish again
-	const EVisibility ScrollBarVisiblityOverride = (AutoWrapText.Get() && CachedSize.IsZero()) ? EVisibility::Collapsed : EVisibility::Visible;
+	// We also hide the scrollbars for non-multi-line text widgets
+	const EVisibility ScrollBarVisiblityOverride = ((AutoWrapText.Get() && CachedSize.IsZero()) || !OwnerWidget->IsMultiLineTextEdit()) ? EVisibility::Collapsed : EVisibility::Visible;
 
 	// Try and make sure that the line containing the cursor is in view
 	if (PositionToScrollIntoView.IsSet())
diff --git a/Engine/Source/Runtime/Slate/Public/Framework/SlateDelegates.h b/Engine/Source/Runtime/Slate/Public/Framework/SlateDelegates.h
--- a/Engine/Source/Runtime/Slate/Public/Framework/SlateDelegates.h
+++ b/Engine/Source/Runtime/Slate/Public/Framework/SlateDelegates.h
@@ -227,4 +227,7 @@ DECLARE_DELEGATE_RetVal_ThreeParams(bool, FAccessSourceCode, const FString& /*Fi
 DECLARE_DELEGATE_RetVal_OneParam(bool, FAccessAsset, UObject* /*InAsset*/);
 
 /** A delegate for a callback during OnKeyDown from an SWidget */
-DECLARE_DELEGATE_RetVal_TwoParams( FReply, FOnKeyDown, const FGeometry&, const FKeyEvent&);
+DECLARE_DELEGATE_RetVal_TwoParams(FReply, FOnKeyDown, const FGeometry&, const FKeyEvent&);
+
+/** A delegate for a callback during OnKeyChar from an SWidget */
+DECLARE_DELEGATE_RetVal_TwoParams(FReply, FOnKeyChar, const FGeometry&, const FCharacterEvent&);
diff --git a/Engine/Source/Runtime/Slate/Public/Widgets/Input/SEditableText.h b/Engine/Source/Runtime/Slate/Public/Widgets/Input/SEditableText.h
--- a/Engine/Source/Runtime/Slate/Public/Widgets/Input/SEditableText.h
+++ b/Engine/Source/Runtime/Slate/Public/Widgets/Input/SEditableText.h
@@ -140,6 +140,9 @@ public:
 		/** Whether to select all text when pressing enter to commit changes */
 		SLATE_ATTRIBUTE( bool, SelectAllTextOnCommit )
 
+		/** Callback delegate to have first chance handling of the OnKeyChar event */
+		SLATE_EVENT(FOnKeyChar, OnKeyCharHandler)
+
 		/** Callback delegate to have first chance handling of the OnKeyDown event */
 		SLATE_EVENT(FOnKeyDown, OnKeyDownHandler)
 
@@ -268,6 +271,15 @@ public:
 
 	/** See Justification attribute */
 	void SetJustification(const TAttribute<ETextJustify::Type>& InJustification);
+	/**
+	 * Sets the OnKeyCharHandler to provide first chance handling of the OnKeyChar event
+	 *
+	 * @param InOnKeyCharHandler			Delegate to call during OnKeyChar event
+	 */
+	void SetOnKeyCharHandler(FOnKeyChar InOnKeyCharHandler)
+	{
+		OnKeyCharHandler = InOnKeyCharHandler;
+	}
 
 	/**
 	 * Sets the OnKeyDownHandler to provide first chance handling of the OnKeyDown event
@@ -437,6 +449,9 @@ protected:
 	/** The iterator to use to detect word boundaries */
 	mutable TSharedPtr<IBreakIterator> WordBreakIterator;
 
+	/** Callback delegate to have first chance handling of the OnKeyChar event */
+	FOnKeyChar OnKeyCharHandler;
+
 	/** Callback delegate to have first chance handling of the OnKeyDown event */
 	FOnKeyDown OnKeyDownHandler;
 
diff --git a/Engine/Source/Runtime/Slate/Public/Widgets/Input/SEditableTextBox.h b/Engine/Source/Runtime/Slate/Public/Widgets/Input/SEditableTextBox.h
--- a/Engine/Source/Runtime/Slate/Public/Widgets/Input/SEditableTextBox.h
+++ b/Engine/Source/Runtime/Slate/Public/Widgets/Input/SEditableTextBox.h
@@ -113,6 +113,9 @@ public:
 		/** Whether to select all text when pressing enter to commit changes */
 		SLATE_ATTRIBUTE( bool, SelectAllTextOnCommit )
 
+		/** Callback delegate to have first chance handling of the OnKeyChar event */
+		SLATE_EVENT(FOnKeyChar, OnKeyCharHandler)
+
 		/** Callback delegate to have first chance handling of the OnKeyDown event */
 		SLATE_EVENT(FOnKeyDown, OnKeyDownHandler)
 
@@ -270,6 +273,13 @@ public:
 	void SetError( const FText& InError );
 	void SetError( const FString& InError );
 
+	/**
+	 * Sets the OnKeyCharHandler to provide first chance handling of the SEditableText's OnKeyChar event
+	 *
+	 * @param InOnKeyCharHandler			Delegate to call during OnKeyChar event
+	 */
+	void SetOnKeyCharHandler(FOnKeyChar InOnKeyCharHandler);
+
 	/**
 	 * Sets the OnKeyDownHandler to provide first chance handling of the SEditableText's OnKeyDown event
 	 *
diff --git a/Engine/Source/Runtime/Slate/Public/Widgets/Input/SMultiLineEditableTextBox.h b/Engine/Source/Runtime/Slate/Public/Widgets/Input/SMultiLineEditableTextBox.h
--- a/Engine/Source/Runtime/Slate/Public/Widgets/Input/SMultiLineEditableTextBox.h
+++ b/Engine/Source/Runtime/Slate/Public/Widgets/Input/SMultiLineEditableTextBox.h
@@ -47,7 +47,7 @@ public:
 		, _Justification(ETextJustify::Left)
 		, _LineHeightPercentage(1.0f)
 		, _IsReadOnly( false )
-		, _IsPassword( false )
+		, _AllowMultiLine( true )
 		, _IsCaretMovedWhenGainFocus ( true )
 		, _SelectAllTextWhenFocused( false )
 		, _ClearTextSelectionOnFocusLoss( true )
@@ -108,8 +108,8 @@ public:
 		/** Sets whether this text box can actually be modified interactively by the user */
 		SLATE_ATTRIBUTE( bool, IsReadOnly )
 
-		/** Sets whether this text box is for storing a password */
-		SLATE_ATTRIBUTE( bool, IsPassword )
+		/** Whether to allow multi-line text */
+		SLATE_ATTRIBUTE(bool, AllowMultiLine)
 
 		/** Workaround as we loose focus when the auto completion closes. */
 		SLATE_ATTRIBUTE( bool, IsCaretMovedWhenGainFocus )
@@ -147,6 +147,14 @@ public:
 		/** Delegate to call before a context menu is opened. User returns the menu content or null to the disable context menu */
 		SLATE_EVENT(FOnContextMenuOpening, OnContextMenuOpening)
 
+		/**
+		 * This is NOT for validating input!
+		 * 
+		 * Called whenever a character is typed.
+		 * Not called for copy, paste, or any other text changes!
+		 */
+		SLATE_EVENT( FOnIsTypedCharValid, OnIsTypedCharValid )
+
 		/** Called whenever the text is changed interactively by the user */
 		SLATE_EVENT( FOnTextChanged, OnTextChanged )
 
@@ -162,6 +170,9 @@ public:
 		/** Called when the cursor is moved within the text area */
 		SLATE_EVENT( SMultiLineEditableText::FOnCursorMoved, OnCursorMoved )
 
+		/** Callback delegate to have first chance handling of the OnKeyChar event */
+		SLATE_EVENT(FOnKeyChar, OnKeyCharHandler)
+
 		/** Callback delegate to have first chance handling of the OnKeyDown event */
 		SLATE_EVENT(FOnKeyDown, OnKeyDownHandler)
 
@@ -386,6 +397,13 @@ public:
 	/** Refresh this text box immediately, rather than wait for the usual caching mechanisms to take affect on the text Tick */
 	void Refresh();
 
+	/**
+	 * Sets the OnKeyCharHandler to provide first chance handling of the SMultiLineEditableText's OnKeyChar event
+	 *
+	 * @param InOnKeyCharHandler			Delegate to call during OnKeyChar event
+	 */
+	void SetOnKeyCharHandler(FOnKeyChar InOnKeyCharHandler);
+
 	/**
 	 * Sets the OnKeyDownHandler to provide first chance handling of the SMultiLineEditableText's OnKeyDown event
 	 *
@@ -453,10 +471,10 @@ protected:
 	/** SomeWidget reporting */
 	TSharedPtr<class IErrorReportingWidget> ErrorReporting;
 
-private:
-
 	const FEditableTextBoxStyle* Style;
 
+private:
+
 	FMargin FORCEINLINE DeterminePadding() const { check(Style);  return PaddingOverride.IsSet() ? PaddingOverride.Get() : Style->Padding; }
 	FMargin FORCEINLINE DetermineHScrollBarPadding() const { check(Style);  return HScrollBarPaddingOverride.IsSet() ? HScrollBarPaddingOverride.Get() : Style->HScrollBarPadding; }
 	FMargin FORCEINLINE DetermineVScrollBarPadding() const { check(Style);  return VScrollBarPaddingOverride.IsSet() ? VScrollBarPaddingOverride.Get() : Style->VScrollBarPadding; }
diff --git a/Engine/Source/Runtime/Slate/Public/Widgets/Text/SMultiLineEditableText.h b/Engine/Source/Runtime/Slate/Public/Widgets/Text/SMultiLineEditableText.h
--- a/Engine/Source/Runtime/Slate/Public/Widgets/Text/SMultiLineEditableText.h
+++ b/Engine/Source/Runtime/Slate/Public/Widgets/Text/SMultiLineEditableText.h
@@ -60,6 +60,7 @@ public:
 		, _IsReadOnly(false)
 		, _OnTextChanged()
 		, _OnTextCommitted()
+		, _AllowMultiLine(true)
 		, _SelectAllTextWhenFocused(false)
 		, _ClearTextSelectionOnFocusLoss(true)
 		, _RevertTextOnEscape(false)
@@ -123,17 +124,28 @@ public:
 		/** The vertical scroll bar widget */
 		SLATE_ARGUMENT(TSharedPtr< SScrollBar >, VScrollBar)
 
+		/**
+		 * This is NOT for validating input!
+		 * 
+		 * Called whenever a character is typed.
+		 * Not called for copy, paste, or any other text changes!
+		 */
+		SLATE_EVENT(FOnIsTypedCharValid, OnIsTypedCharValid)
+
 		/** Called whenever the text is changed interactively by the user */
 		SLATE_EVENT(FOnTextChanged, OnTextChanged)
 
 		/** Called whenever the text is committed.  This happens when the user presses enter or the text box loses focus. */
 		SLATE_EVENT(FOnTextCommitted, OnTextCommitted)
 
+		/** Whether to allow multi-line text */
+		SLATE_ATTRIBUTE(bool, AllowMultiLine)
+
 		/** Whether to select all text when the user clicks to give focus on the widget */
 		SLATE_ATTRIBUTE(bool, SelectAllTextWhenFocused)
 
 		/** Whether to clear text selection when focus is lost */
-		SLATE_ATTRIBUTE( bool, ClearTextSelectionOnFocusLoss )
+		SLATE_ATTRIBUTE(bool, ClearTextSelectionOnFocusLoss)
 
 		/** Whether to allow the user to back out of changes when they press the escape key */
 		SLATE_ATTRIBUTE(bool, RevertTextOnEscape)
@@ -156,6 +168,9 @@ public:
 		/** Called when the cursor is moved within the text area */
 		SLATE_EVENT(FOnCursorMoved, OnCursorMoved)
 
+		/** Callback delegate to have first chance handling of the OnKeyChar event */
+		SLATE_EVENT(FOnKeyChar, OnKeyCharHandler)
+
 		/** Callback delegate to have first chance handling of the OnKeyDown event */
 		SLATE_EVENT(FOnKeyDown, OnKeyDownHandler)
 
@@ -338,6 +353,16 @@ public:
 	/** Refresh this editable text immediately, rather than wait for the usual caching mechanisms to take affect on the text Tick */
 	void Refresh();
 
+	/**
+	 * Sets the OnKeyCharHandler to provide first chance handling of the OnKeyChar event
+	 *
+	 * @param InOnKeyCharHandler			Delegate to call during OnKeyChar event
+	 */
+	void SetOnKeyCharHandler(FOnKeyChar InOnKeyCharHandler)
+	{
+		OnKeyCharHandler = InOnKeyCharHandler;
+	}
+
 	/**
 	 * Sets the OnKeyDownHandler to provide first chance handling of the OnKeyDown event
 	 *
@@ -419,6 +444,9 @@ protected:
 	/** The text layout that deals with the editable text */
 	TUniquePtr<FSlateEditableTextLayout> EditableTextLayout;
 
+	/** Whether to allow multi-line text */
+	TAttribute<bool> bAllowMultiLine;
+
 	/** Whether to select all text when the user clicks to give focus on the widget */
 	TAttribute<bool> bSelectAllTextWhenFocused;
 
@@ -440,6 +468,9 @@ protected:
 	/** Delegate to call before a context menu is opened */
 	FOnContextMenuOpening OnContextMenuOpening;
 
+	/** Called when a character is typed and we want to know if the text field supports typing this character. */
+	FOnIsTypedCharValid OnIsTypedCharValid;
+
 	/** Called whenever the text is changed interactively by the user */
 	FOnTextChanged OnTextChangedCallback;
 
@@ -479,6 +510,9 @@ protected:
 	/**	The current position of the software cursor */
 	FVector2D SoftwareCursorPosition;
 
+	/** Callback delegate to have first chance handling of the OnKeyChar event */
+	FOnKeyChar OnKeyCharHandler;
+
 	/** Callback delegate to have first chance handling of the OnKeyDown event */
 	FOnKeyDown OnKeyDownHandler;
 
diff --git a/Engine/Source/Runtime/SlateCore/Private/Fonts/SlateTextShaper.cpp b/Engine/Source/Runtime/SlateCore/Private/Fonts/SlateTextShaper.cpp
--- a/Engine/Source/Runtime/SlateCore/Private/Fonts/SlateTextShaper.cpp
+++ b/Engine/Source/Runtime/SlateCore/Private/Fonts/SlateTextShaper.cpp
@@ -315,7 +315,13 @@ void FSlateTextShaper::PerformKerningOnlyTextShaping(const TCHAR* InText, const
 		const int32 TextEndIndex = InTextStart + InTextLen;
 		for (; RunningTextIndex < TextEndIndex; ++RunningTextIndex)
 		{
-			const TCHAR CurrentChar = InText[RunningTextIndex];
+			TCHAR CurrentChar = InText[RunningTextIndex];
+
+			// Substitute whitespace characters with spaces since not all fonts support all kinds of whitespace characters (but we don't care since we don't render them anyway)
+			if (FText::IsWhitespace(CurrentChar))
+			{
+				CurrentChar = TEXT(' ');
+			}
 
 			// First try with the actual character
 			float SubFontScalingFactor = 1.0f;
@@ -462,7 +468,13 @@ void FSlateTextShaper::PerformHarfBuzzTextShaping(const TCHAR* InText, const int
 		const int32 TextEndIndex = InTextStart + InTextLen;
 		for (; RunningTextIndex < TextEndIndex; ++RunningTextIndex)
 		{
-			const TCHAR CurrentChar = InText[RunningTextIndex];
+			TCHAR CurrentChar = InText[RunningTextIndex];
+
+			// Substitute whitespace characters with spaces since not all fonts support all kinds of whitespace characters (but we don't care since we don't render them anyway)
+			if (FText::IsWhitespace(CurrentChar))
+			{
+				CurrentChar = TEXT(' ');
+			}
 
 			// First try with the actual character
 			float SubFontScalingFactor = 1.0f;
@@ -790,22 +802,22 @@ bool FSlateTextShaper::InsertSubstituteGlyphs(const TCHAR* InText, const int32 I
 		}
 #endif // WITH_FREETYPE
 
-		// We insert (up-to) 4 space glyphs in-place of a tab character
+		// We insert a spacer glyph with (up-to) the width of 4 space glyphs in-place of a tab character
 		const int32 NumSpacesToInsert = 4 - (OutGlyphsToRender.Num() % 4);
-		for (int32 SpaceIndex = 0; SpaceIndex < NumSpacesToInsert; ++SpaceIndex)
+		if (NumSpacesToInsert > 0)
 		{
 			const int32 CurrentGlyphEntryIndex = OutGlyphsToRender.AddDefaulted();
 			FShapedGlyphEntry& ShapedGlyphEntry = OutGlyphsToRender[CurrentGlyphEntryIndex];
 			ShapedGlyphEntry.FontFaceData = InShapedGlyphFaceData;
 			ShapedGlyphEntry.GlyphIndex = SpaceGlyphIndex;
 			ShapedGlyphEntry.SourceIndex = InCharIndex;
-			ShapedGlyphEntry.XAdvance = SpaceXAdvance;
+			ShapedGlyphEntry.XAdvance = SpaceXAdvance * NumSpacesToInsert;
 			ShapedGlyphEntry.YAdvance = 0;
 			ShapedGlyphEntry.XOffset = 0;
 			ShapedGlyphEntry.YOffset = 0;
 			ShapedGlyphEntry.Kerning = 0;
-			ShapedGlyphEntry.NumCharactersInGlyph = (SpaceIndex == 0) ? 1 : 0;
-			ShapedGlyphEntry.NumGraphemeClustersInGlyph = (SpaceIndex == 0) ? 1 : 0;
+			ShapedGlyphEntry.NumCharactersInGlyph = 1;
+			ShapedGlyphEntry.NumGraphemeClustersInGlyph = 1;
 			ShapedGlyphEntry.TextDirection = TextBiDi::ETextDirection::LeftToRight;
 			ShapedGlyphEntry.bIsVisible = false;
 		}
