Add ability to auto select the default device to SDeviceOutputLog and Fix static initialization shutdown bug in SOutputLog.

diff --git a/Engine/Source/Developer/OutputLog/Private/OutputLogModule.cpp b/Engine/Source/Developer/OutputLog/Private/OutputLogModule.cpp
--- a/Engine/Source/Developer/OutputLog/Private/OutputLogModule.cpp
+++ b/Engine/Source/Developer/OutputLog/Private/OutputLogModule.cpp
@@ -112,6 +112,8 @@ void FOutputLogModule::ShutdownModule()
 		FGlobalTabmanager::Get()->UnregisterNomadTabSpawner(OutputLogModule::OutputLogTabName);
 		FGlobalTabmanager::Get()->UnregisterNomadTabSpawner(OutputLogModule::DeviceOutputLogTabName);
 	}
+
+	OutputLogHistory.Reset();
 }
 
 TSharedRef< SWidget > FOutputLogModule::MakeConsoleInputBox( TSharedPtr< SMultiLineEditableTextBox >& OutExposedEditableTextBox ) const
diff --git a/Engine/Source/Developer/OutputLog/Private/SDeviceOutputLog.cpp b/Engine/Source/Developer/OutputLog/Private/SDeviceOutputLog.cpp
--- a/Engine/Source/Developer/OutputLog/Private/SDeviceOutputLog.cpp
+++ b/Engine/Source/Developer/OutputLog/Private/SDeviceOutputLog.cpp
@@ -24,6 +24,8 @@ static bool IsSupportedPlatform(ITargetPlatform* Platform)
 
 void SDeviceOutputLog::Construct( const FArguments& InArgs )
 {
+	bAutoSelectDevice = InArgs._AutoSelectDevice;
+
 	MessagesTextMarshaller = FOutputLogTextLayoutMarshaller::Create(TArray<TSharedPtr<FLogMessage>>(), &Filter);
 
 	MessagesTextBox = SNew(SMultiLineEditableTextBox)
@@ -144,10 +146,38 @@ SDeviceOutputLog::~SDeviceOutputLog()
 			Platform->OnDeviceLost().RemoveAll(this);
 		}
 	}
+	CurrentDeviceOutputPtr.Reset();
 }
 
 void SDeviceOutputLog::Tick(const FGeometry& AllottedGeometry, const double InCurrentTime, const float InDeltaTime)
 {
+	// If auto-select is enabled request connecting to the default device and select it
+	if (!CurrentDevicePtr.IsValid() && bAutoSelectDevice)
+	{
+		int32 DefaultDeviceEntryIdx = DeviceList.IndexOfByPredicate([this](const TSharedPtr<FTargetDeviceEntry>& Other) {
+			ITargetDevicePtr PinnedPtr = Other->DeviceWeakPtr.Pin();
+			return PinnedPtr->IsDefault();
+		});
+
+		if (DeviceList.IsValidIndex(DefaultDeviceEntryIdx))
+		{
+			ITargetDevicePtr PinnedPtr = DeviceList[DefaultDeviceEntryIdx]->DeviceWeakPtr.Pin();
+			PinnedPtr->Connect();
+			OnDeviceSelectionChanged(DeviceList[DefaultDeviceEntryIdx]);
+		}
+	}
+
+	// If the device is selected but was not yet connected then the output router would not have been registered
+	if (CurrentDevicePtr.IsValid() && !CurrentDeviceOutputPtr.IsValid())
+	{
+		ITargetDevicePtr PinnedPtr = CurrentDevicePtr->DeviceWeakPtr.Pin();
+		if (PinnedPtr.IsValid() && PinnedPtr->IsConnected())
+		{
+			// It is now connected so register the output router
+			CurrentDeviceOutputPtr = PinnedPtr->CreateDeviceOutputRouter(this);
+		}
+	}
+
 	FScopeLock ScopeLock(&BufferedLinesSynch);
 	if (BufferedLines.Num() > 0)
 	{
@@ -227,6 +257,16 @@ void SDeviceOutputLog::HandleTargetPlatformDeviceDiscovered(ITargetDeviceRef Dis
 	}
 }
 
+ITargetDevicePtr SDeviceOutputLog::GetSelectedTargetDevice() const
+{
+	ITargetDevicePtr PinnedPtr = nullptr;
+	if (CurrentDevicePtr.IsValid())
+	{
+		PinnedPtr = CurrentDevicePtr->DeviceWeakPtr.Pin();
+	}
+	return PinnedPtr;
+}
+
 void SDeviceOutputLog::AddDeviceEntry(ITargetDeviceRef TargetDevice)
 {
 	if (FindDeviceEntry(TargetDevice->GetId()))
diff --git a/Engine/Source/Developer/OutputLog/Private/SDeviceOutputLog.h b/Engine/Source/Developer/OutputLog/Private/SDeviceOutputLog.h
--- a/Engine/Source/Developer/OutputLog/Private/SDeviceOutputLog.h
+++ b/Engine/Source/Developer/OutputLog/Private/SDeviceOutputLog.h
@@ -24,12 +24,14 @@ typedef TSharedPtr<FTargetDeviceEntry> FTargetDeviceEntryPtr;
 class SDeviceOutputLog : public SOutputLog
 {
 public:
-	SLATE_BEGIN_ARGS( SDeviceOutputLog )
+	SLATE_BEGIN_ARGS(SDeviceOutputLog) 
+		: _AutoSelectDevice(false)
 	{}
+		SLATE_ARGUMENT(bool, AutoSelectDevice)	// If the first device found should be automatically selected.
 	SLATE_END_ARGS()
 
 	/** Destructor for output log, so we can unregister from notifications */
-	~SDeviceOutputLog();
+	virtual ~SDeviceOutputLog();
 
 	/**
 	 * Construct this widget.  Called by the SNew() Slate macro.
@@ -38,6 +40,8 @@ public:
 	 */
 	void Construct( const FArguments& InArgs );
 
+	ITargetDevicePtr GetSelectedTargetDevice() const;
+
 protected:
 	// SWidget interface
 	virtual void Tick(const FGeometry& AllottedGeometry, const double InCurrentTime, const float InDeltaTime) override;
@@ -76,4 +80,6 @@ private:
 	/** Synchronization object for access to buffered lines */
 	FCriticalSection		BufferedLinesSynch;
 	TArray<FBufferedLine>	BufferedLines;
+
+	bool bAutoSelectDevice;
 };
diff --git a/Engine/Source/Developer/OutputLog/Private/SOutputLog.h b/Engine/Source/Developer/OutputLog/Private/SOutputLog.h
--- a/Engine/Source/Developer/OutputLog/Private/SOutputLog.h
+++ b/Engine/Source/Developer/OutputLog/Private/SOutputLog.h
@@ -319,7 +319,7 @@ public:
 	SLATE_END_ARGS()
 
 	/** Destructor for output log, so we can unregister from notifications */
-	~SOutputLog();
+	virtual ~SOutputLog();
 
 	/**
 	 * Construct this widget.  Called by the SNew() Slate macro.
