UE-69183 - iOS Device Output Log doesn't recover previously disconnected device

diff --git a/Engine/Source/Developer/OutputLog/Private/SDeviceOutputLog.cpp b/Engine/Source/Developer/OutputLog/Private/SDeviceOutputLog.cpp
--- a/Engine/Source/Developer/OutputLog/Private/SDeviceOutputLog.cpp
+++ b/Engine/Source/Developer/OutputLog/Private/SDeviceOutputLog.cpp
@@ -250,6 +250,20 @@ void SDeviceOutputLog::HandleTargetPlatformDeviceDiscovered(ITargetDeviceRef Dis
 	if (DeviceList.IsValidIndex(ExistingEntryIdx))
 	{
 		DeviceList[ExistingEntryIdx]->DeviceWeakPtr = DiscoveredDevice;
+		if (CurrentDevicePtr == DeviceList[ExistingEntryIdx])
+		{
+			if (!CurrentDeviceOutputPtr.IsValid())
+			{
+				if (CurrentDevicePtr.IsValid())
+				{
+					ITargetDevicePtr PinnedPtr = CurrentDevicePtr->DeviceWeakPtr.Pin();
+					if (PinnedPtr.IsValid() && PinnedPtr->IsConnected())
+					{
+						CurrentDeviceOutputPtr = PinnedPtr->CreateDeviceOutputRouter(this);
+					}
+				}
+			}
+		}
 	}
 	else
 	{
