UE-64603 - Android Devices do not show up in device list for Device Output Log

diff --git a/Engine/Source/Developer/OutputLog/Private/SDeviceOutputLog.cpp b/Engine/Source/Developer/OutputLog/Private/SDeviceOutputLog.cpp
--- a/Engine/Source/Developer/OutputLog/Private/SDeviceOutputLog.cpp
+++ b/Engine/Source/Developer/OutputLog/Private/SDeviceOutputLog.cpp
@@ -238,10 +238,6 @@ void SDeviceOutputLog::AddDeviceEntry(ITargetDeviceRef TargetDevice)
 	{
 		return;
 	}
-	if (TargetDevice->IsPlatformAggregated())
-	{
-		return;
-	}
 	const FString DummyIOSDeviceName(FString::Printf(TEXT("All_iOS_On_%s"), FPlatformProcess::ComputerName()));
 	const FString DummyTVOSDeviceName(FString::Printf(TEXT("All_tvOS_On_%s"), FPlatformProcess::ComputerName()));
 	if (TargetDevice->GetId().GetDeviceName().Equals(DummyIOSDeviceName, ESearchCase::IgnoreCase) ||
