Updated tab key functionality in the output log's console box to not leave behind unwanted tab characters

diff --git a/Engine/Source/Developer/OutputLog/Private/SOutputLog.cpp b/Engine/Source/Developer/OutputLog/Private/SOutputLog.cpp
--- a/Engine/Source/Developer/OutputLog/Private/SOutputLog.cpp
+++ b/Engine/Source/Developer/OutputLog/Private/SOutputLog.cpp
@@ -305,7 +305,7 @@ FReply SConsoleInputBox::OnPreviewKeyDown(const FGeometry& MyGeometry, const FKe
 				if (Suggestions.HasSelectedSuggestion())
 				{
 					MarkActiveSuggestion();
-					OnTextCommitted(InputText->GetText(), ETextCommit::OnEnter);
+					SuggestionBox->SetIsOpen(false);
 				}
 				else
 				{
@@ -314,6 +314,7 @@ FReply SConsoleInputBox::OnPreviewKeyDown(const FGeometry& MyGeometry, const FKe
 				}
 			}
 
+			bConsumeTab = true;
 			return FReply::Handled();
 		}
 		else if (KeyEvent.GetKey() == EKeys::Escape)
@@ -578,6 +579,13 @@ FReply SConsoleInputBox::OnKeyCharHandler(const FGeometry& MyGeometry, const FCh
 		return FReply::Handled();
 	}
 
+	// Intercept tab if used for auto-complete
+	if (InCharacterEvent.GetCharacter() == '\t' && bConsumeTab)
+	{
+		bConsumeTab = false;
+		return FReply::Handled();
+	}
+
 	FInputChord OpenConsoleChord;
 	if (ActiveCommandExecutor && ActiveCommandExecutor->AllowHotKeyClose())
 	{
diff --git a/Engine/Source/Developer/OutputLog/Private/SOutputLog.h b/Engine/Source/Developer/OutputLog/Private/SOutputLog.h
--- a/Engine/Source/Developer/OutputLog/Private/SOutputLog.h
+++ b/Engine/Source/Developer/OutputLog/Private/SOutputLog.h
@@ -230,6 +230,10 @@ private:
 
 	/** true if this widget has been Ticked at least once */
 	bool bHasTicked;
+
+	/** True if we consumed a tab key in OnPreviewKeyDown, so we can ignore it in OnKeyCharHandler as well */
+	bool bConsumeTab;
+
 };
 
 /**
