Swapped radio button for check button to fix UI glitch

#rnx

diff --git a/Engine/Source/Developer/OutputLog/Private/SOutputLog.cpp b/Engine/Source/Developer/OutputLog/Private/SOutputLog.cpp
--- a/Engine/Source/Developer/OutputLog/Private/SOutputLog.cpp
+++ b/Engine/Source/Developer/OutputLog/Private/SOutputLog.cpp
@@ -572,7 +572,7 @@ TSharedRef<SWidget> SConsoleInputBox::GetCommandExecutorMenuContent()
 					FIsActionChecked::CreateLambda([bIsActiveCmdExec] { return bIsActiveCmdExec; })
 					),
 				NAME_None,
-				EUserInterfaceActionType::RadioButton
+				EUserInterfaceActionType::Check
 			);
 		}
 	}
