PR #6273: Fix memory leak in SOutputLog (Contributed by [REDACTED_PERSON])

diff --git a/Engine/Source/Developer/OutputLog/Private/SOutputLog.cpp b/Engine/Source/Developer/OutputLog/Private/SOutputLog.cpp
--- a/Engine/Source/Developer/OutputLog/Private/SOutputLog.cpp
+++ b/Engine/Source/Developer/OutputLog/Private/SOutputLog.cpp
@@ -1555,6 +1555,6 @@ bool FOutputLogFilter::IsLogCategoryEnabled(const FName& LogCategory) const
 void FOutputLogFilter::ClearSelectedLogCategories()
 {
 	// No need to churn memory each time the selected categories are cleared
-	SelectedLogCategories.Reset(SelectedLogCategories.GetAllocatedSize());
+	SelectedLogCategories.Reset();
 }
 #undef LOCTEXT_NAMESPACE
