Defer Output Log text layout modification to Tick

This avoids the text layout for the Output Log being modified while the Output Log itself is being painted (eg, from messages being logged during OnPaint).

diff --git a/Engine/Source/Developer/OutputLog/Private/SDeviceOutputLog.cpp b/Engine/Source/Developer/OutputLog/Private/SDeviceOutputLog.cpp
--- a/Engine/Source/Developer/OutputLog/Private/SDeviceOutputLog.cpp
+++ b/Engine/Source/Developer/OutputLog/Private/SDeviceOutputLog.cpp
@@ -183,17 +183,12 @@ void SDeviceOutputLog::Tick(const FGeometry& AllottedGeometry, const double InCu
 	{
 		for (const FBufferedLine& Line : BufferedLines)
 		{
-			MessagesTextMarshaller->AppendMessage(Line.Data, Line.Verbosity, Line.Category);
+			MessagesTextMarshaller->AppendPendingMessage(Line.Data, Line.Verbosity, Line.Category);
 		}
-
-		// Don't scroll to the bottom automatically when the user is scrolling the view or has scrolled it away from the bottom.
-		if (!bIsUserScrolled)
-		{
-			MessagesTextBox->ScrollTo(FTextLocation(MessagesTextMarshaller->GetNumMessages() - 1));
-		}
-
 		BufferedLines.Empty(32);
 	}
+
+	SOutputLog::Tick(AllottedGeometry, InCurrentTime, InDeltaTime);
 }
 
 void SDeviceOutputLog::Serialize(const TCHAR* V, ELogVerbosity::Type Verbosity, const class FName& Category)
diff --git a/Engine/Source/Developer/OutputLog/Private/SOutputLog.cpp b/Engine/Source/Developer/OutputLog/Private/SOutputLog.cpp
--- a/Engine/Source/Developer/OutputLog/Private/SOutputLog.cpp
+++ b/Engine/Source/Developer/OutputLog/Private/SOutputLog.cpp
@@ -658,7 +658,8 @@ FOutputLogTextLayoutMarshaller::~FOutputLogTextLayoutMarshaller()
 void FOutputLogTextLayoutMarshaller::SetText(const FString& SourceString, FTextLayout& TargetTextLayout)
 {
 	TextLayout = &TargetTextLayout;
-	AppendMessagesToTextLayout(Messages);
+	NextPendingMessageIndex = 0;
+	SubmitPendingMessages();
 }
 
 void FOutputLogTextLayoutMarshaller::GetText(FString& TargetString, const FTextLayout& SourceTextLayout)
@@ -666,86 +667,69 @@ void FOutputLogTextLayoutMarshaller::GetText(FString& TargetString, const FTextL
 	SourceTextLayout.GetAsText(TargetString);
 }
 
-bool FOutputLogTextLayoutMarshaller::AppendMessage(const TCHAR* InText, const ELogVerbosity::Type InVerbosity, const FName& InCategory)
+bool FOutputLogTextLayoutMarshaller::AppendPendingMessage(const TCHAR* InText, const ELogVerbosity::Type InVerbosity, const FName& InCategory)
 {
-	TArray< TSharedPtr<FOutputLogMessage> > NewMessages;
-	if(SOutputLog::CreateLogMessages(InText, InVerbosity, InCategory, NewMessages))
-	{
-		const bool bWasEmpty = Messages.Num() == 0;
-		Messages.Append(NewMessages);
-
-		// Add new message categories to the filter's available log categories
-		for (const auto& NewMessage : NewMessages)
-		{
-			Filter->AddAvailableLogCategory(NewMessage->Category);
-		}
-
-		if(TextLayout)
-		{
-			// If we were previously empty, then we'd have inserted a dummy empty line into the document
-			// We need to remove this line now as it would cause the message indices to get out-of-sync with the line numbers, which would break auto-scrolling
-			if(bWasEmpty)
-			{
-				TextLayout->ClearLines();
-			}
-
-			// If we've already been given a text layout, then append these new messages rather than force a refresh of the entire document
-			AppendMessagesToTextLayout(NewMessages);
-		}
-		else
-		{
-			MarkMessagesCacheAsDirty();
-			MakeDirty();
-		}
+	return SOutputLog::CreateLogMessages(InText, InVerbosity, InCategory, Messages);
+}
 
+bool FOutputLogTextLayoutMarshaller::SubmitPendingMessages()
+{
+	if (Messages.IsValidIndex(NextPendingMessageIndex))
+	{
+		const int32 CurrentMessagesCount = Messages.Num();
+		AppendPendingMessagesToTextLayout();
+		NextPendingMessageIndex = CurrentMessagesCount;
 		return true;
 	}
-
 	return false;
 }
 
-void FOutputLogTextLayoutMarshaller::AppendMessageToTextLayout(const TSharedPtr<FOutputLogMessage>& InMessage)
+void FOutputLogTextLayoutMarshaller::AppendPendingMessagesToTextLayout()
 {
-	if (!Filter->IsMessageAllowed(InMessage))
+	const int32 CurrentMessagesCount = Messages.Num();
+	const int32 NumPendingMessages = CurrentMessagesCount - NextPendingMessageIndex;
+
+	if (NumPendingMessages == 0)
 	{
 		return;
 	}
 
-	// Increment the cached count if we're not rebuilding the log
-	if ( !IsDirty() )
+	if (TextLayout)
 	{
-		CachedNumMessages++;
+		// If we were previously empty, then we'd have inserted a dummy empty line into the document
+		// We need to remove this line now as it would cause the message indices to get out-of-sync with the line numbers, which would break auto-scrolling
+		const bool bWasEmpty = GetNumMessages() == 0;
+		if (bWasEmpty)
+		{
+			TextLayout->ClearLines();
+		}
+	}
+	else
+	{
+		MarkMessagesCacheAsDirty();
+		MakeDirty();
 	}
 
-	const FTextBlockStyle& MessageTextStyle = FEditorStyle::Get().GetWidgetStyle<FTextBlockStyle>(InMessage->Style);
-
-	TSharedRef<FString> LineText = InMessage->Message;
-
-	TArray<TSharedRef<IRun>> Runs;
-	Runs.Add(FSlateTextRun::Create(FRunInfo(), LineText, MessageTextStyle));
-
-	TextLayout->AddLine(FSlateTextLayout::FNewLineData(MoveTemp(LineText), MoveTemp(Runs)));
-}
-
-void FOutputLogTextLayoutMarshaller::AppendMessagesToTextLayout(const TArray<TSharedPtr<FOutputLogMessage>>& InMessages)
-{
 	TArray<FTextLayout::FNewLineData> LinesToAdd;
-	LinesToAdd.Reserve(InMessages.Num());
+	LinesToAdd.Reserve(NumPendingMessages);
 
 	int32 NumAddedMessages = 0;
 
-	for (const auto& CurrentMessage : InMessages)
+	for (int32 MessageIndex = NextPendingMessageIndex; MessageIndex < CurrentMessagesCount; ++MessageIndex)
 	{
-		if (!Filter->IsMessageAllowed(CurrentMessage))
+		const TSharedPtr<FOutputLogMessage> Message = Messages[MessageIndex];
+
+		Filter->AddAvailableLogCategory(Message->Category);
+		if (!Filter->IsMessageAllowed(Message))
 		{
 			continue;
 		}
 
 		++NumAddedMessages;
 
-		const FTextBlockStyle& MessageTextStyle = FEditorStyle::Get().GetWidgetStyle<FTextBlockStyle>(CurrentMessage->Style);
+		const FTextBlockStyle& MessageTextStyle = FEditorStyle::Get().GetWidgetStyle<FTextBlockStyle>(Message->Style);
 
-		TSharedRef<FString> LineText = CurrentMessage->Message;
+		TSharedRef<FString> LineText = Message->Message;
 
 		TArray<TSharedRef<IRun>> Runs;
 		Runs.Add(FSlateTextRun::Create(FRunInfo(), LineText, MessageTextStyle));
@@ -764,6 +748,7 @@ void FOutputLogTextLayoutMarshaller::AppendMessagesToTextLayout(const TArray<TSh
 
 void FOutputLogTextLayoutMarshaller::ClearMessages()
 {
+	NextPendingMessageIndex = 0;
 	Messages.Empty();
 	MakeDirty();
 }
@@ -778,8 +763,9 @@ void FOutputLogTextLayoutMarshaller::CountMessages()
 
 	CachedNumMessages = 0;
 
-	for (const auto& CurrentMessage : Messages)
+	for (int32 MessageIndex = 0; MessageIndex < NextPendingMessageIndex; ++MessageIndex)
 	{
+		const TSharedPtr<FOutputLogMessage> CurrentMessage = Messages[MessageIndex];
 		if (Filter->IsMessageAllowed(CurrentMessage))
 		{
 			CachedNumMessages++;
@@ -792,7 +778,8 @@ void FOutputLogTextLayoutMarshaller::CountMessages()
 
 int32 FOutputLogTextLayoutMarshaller::GetNumMessages() const
 {
-	return Messages.Num();
+	const int32 NumPendingMessages = Messages.Num() - NextPendingMessageIndex;
+	return Messages.Num() - NumPendingMessages;
 }
 
 int32 FOutputLogTextLayoutMarshaller::GetNumFilteredMessages()
@@ -819,6 +806,7 @@ void FOutputLogTextLayoutMarshaller::MarkMessagesCacheAsDirty()
 
 FOutputLogTextLayoutMarshaller::FOutputLogTextLayoutMarshaller(TArray< TSharedPtr<FOutputLogMessage> > InMessages, FOutputLogFilter* InFilter)
 	: Messages(MoveTemp(InMessages))
+	, NextPendingMessageIndex(0)
 	, CachedNumMessages(0)
 	, Filter(InFilter)
 	, TextLayout(nullptr)
@@ -988,6 +976,20 @@ SOutputLog::~SOutputLog()
 	FCoreDelegates::OnHandleSystemError.RemoveAll(this);
 }
 
+void SOutputLog::Tick(const FGeometry& AllottedGeometry, const double InCurrentTime, const float InDeltaTime)
+{
+	if (MessagesTextMarshaller->SubmitPendingMessages())
+	{
+		// Don't scroll to the bottom automatically when the user is scrolling the view or has scrolled it away from the bottom.
+		if (!bIsUserScrolled)
+		{
+			RequestForceScroll();
+		}
+	}
+
+	SCompoundWidget::Tick(AllottedGeometry, InCurrentTime, InDeltaTime);
+}
+
 void SOutputLog::OnCrash()
 {
 	if (GLog != nullptr)
@@ -1082,14 +1084,7 @@ bool SOutputLog::CreateLogMessages( const TCHAR* V, ELogVerbosity::Type Verbosit
 
 void SOutputLog::Serialize(const TCHAR* V, ELogVerbosity::Type Verbosity, const class FName& Category)
 {
-	if ( MessagesTextMarshaller->AppendMessage(V, Verbosity, Category) )
-	{
-		// Don't scroll to the bottom automatically when the user is scrolling the view or has scrolled it away from the bottom.
-		if( !bIsUserScrolled )
-		{
-			RequestForceScroll();
-		}
-	}
+	MessagesTextMarshaller->AppendPendingMessage(V, Verbosity, Category);
 }
 
 FSlateColor SOutputLog::GetViewButtonForegroundColor() const
@@ -1137,6 +1132,8 @@ bool SOutputLog::CanClearLog() const
 
 void SOutputLog::OnConsoleCommandExecuted()
 {
+	// Submit pending messages when executing a command to keep the log feeling responsive to input
+	MessagesTextMarshaller->SubmitPendingMessages();
 	RequestForceScroll();
 }
 
diff --git a/Engine/Source/Developer/OutputLog/Private/SOutputLog.h b/Engine/Source/Developer/OutputLog/Private/SOutputLog.h
--- a/Engine/Source/Developer/OutputLog/Private/SOutputLog.h
+++ b/Engine/Source/Developer/OutputLog/Private/SOutputLog.h
@@ -328,6 +328,9 @@ public:
 	 */
 	void Construct( const FArguments& InArgs );
 
+	// SWidget interface
+	virtual void Tick(const FGeometry& AllottedGeometry, const double InCurrentTime, const float InDeltaTime) override;
+
 	/**
 	 * Creates FOutputLogMessage objects from FOutputDevice log callback
 	 *
@@ -469,7 +472,8 @@ public:
 	virtual void SetText(const FString& SourceString, FTextLayout& TargetTextLayout) override;
 	virtual void GetText(FString& TargetString, const FTextLayout& SourceTextLayout) override;
 
-	bool AppendMessage(const TCHAR* InText, const ELogVerbosity::Type InVerbosity, const FName& InCategory);
+	bool AppendPendingMessage(const TCHAR* InText, const ELogVerbosity::Type InVerbosity, const FName& InCategory);
+	bool SubmitPendingMessages();
 	void ClearMessages();
 
 	void CountMessages();
@@ -483,12 +487,14 @@ protected:
 
 	FOutputLogTextLayoutMarshaller(TArray< TSharedPtr<FOutputLogMessage> > InMessages, FOutputLogFilter* InFilter);
 
-	void AppendMessageToTextLayout(const TSharedPtr<FOutputLogMessage>& InMessage);
-	void AppendMessagesToTextLayout(const TArray<TSharedPtr<FOutputLogMessage>>& InMessages);
+	void AppendPendingMessagesToTextLayout();
 
 	/** All log messages to show in the text box */
 	TArray< TSharedPtr<FOutputLogMessage> > Messages;
 
+	/** Index of the next entry in the Messages array that is pending submission to the text layout */
+	int32 NextPendingMessageIndex;
+
 	/** Holds cached numbers of messages to avoid unnecessary re-filtering */
 	int32 CachedNumMessages;
 	
