PR #5593: Word wrapping in the Output log window (Contributed by [REDACTED_PERSON])

diff --git a/Engine/Config/BaseEditorPerProjectUserSettings.ini b/Engine/Config/BaseEditorPerProjectUserSettings.ini
--- a/Engine/Config/BaseEditorPerProjectUserSettings.ini
+++ b/Engine/Config/BaseEditorPerProjectUserSettings.ini
@@ -19,6 +19,8 @@
 ; -------------------------------------------------------------------------------------------------
 
 [/Script/UnrealEd.EditorPerProjectUserSettings]
+bEnableOutputLogWordWrap=False
+
 ; True if WASD keys should be remapped to flight controls while the right mouse button is held down
 AllowFlightCameraToRemapKeys=True
 
diff --git a/Engine/Source/Developer/OutputLog/Private/SOutputLog.cpp b/Engine/Source/Developer/OutputLog/Private/SOutputLog.cpp
--- a/Engine/Source/Developer/OutputLog/Private/SOutputLog.cpp
+++ b/Engine/Source/Developer/OutputLog/Private/SOutputLog.cpp
@@ -20,6 +20,8 @@
 #include "EditorStyleSet.h"
 #include "Classes/EditorStyleSettings.h"
 #include "Widgets/Input/SSearchBox.h"
+#include "Widgets/Input/SCheckBox.h"
+#include "Widgets/Images/SImage.h"
 #include "Features/IModularFeatures.h"
 #include "Misc/CoreDelegates.h"
 
@@ -832,7 +834,6 @@ void SOutputLog::Construct( const FArguments& InArgs )
 
 	MessagesTextMarshaller = FOutputLogTextLayoutMarshaller::Create(InArgs._Messages, &Filter);
 
-
 	MessagesTextBox = SNew(SMultiLineEditableTextBox)
 		.Style(FEditorStyle::Get(), "Log.TextBox")
 		.TextStyle(FEditorStyle::Get(), "Log.Normal")
@@ -840,6 +841,7 @@ void SOutputLog::Construct( const FArguments& InArgs )
 		.Marshaller(MessagesTextMarshaller)
 		.IsReadOnly(true)
 		.AlwaysShowScrollbars(true)
+		.AutoWrapText(this, &SOutputLog::IsWordWrapEnabled)
 		.OnVScrollBarUserScrolled(this, &SOutputLog::OnUserScrolled)
 		.ContextMenuExtender(this, &SOutputLog::ExtendTextBoxMenu);
 
@@ -914,16 +916,53 @@ void SOutputLog::Construct( const FArguments& InArgs )
 			// The console input box
 			+SVerticalBox::Slot()
 			.AutoHeight()
-			.Padding(FMargin(0.0f, 4.0f, 0.0f, 0.0f))
 			[
-				SNew(SBox)
-				.MaxDesiredHeight(180.0f)
+				SNew(SHorizontalBox)
+
+				+SHorizontalBox::Slot()
+				.FillWidth(1.f)
+				.VAlign(VAlign_Center)
+				.Padding(FMargin(0.0f, 1.0f, 0.0f, 0.0f))
 				[
-					SNew(SConsoleInputBox)
-					.OnConsoleCommandExecuted(this, &SOutputLog::OnConsoleCommandExecuted)
+					SNew(SBox)
+					.MaxDesiredHeight(180.0f)
+					[
+						SNew(SConsoleInputBox)
+						.OnConsoleCommandExecuted(this, &SOutputLog::OnConsoleCommandExecuted)
 
-					// Always place suggestions above the input line for the output log widget
-					.SuggestionListPlacement(MenuPlacement_AboveAnchor)
+						// Always place suggestions above the input line for the output log widget
+						.SuggestionListPlacement(MenuPlacement_AboveAnchor)
+					]
+				]
+
+				+SHorizontalBox::Slot()
+				.AutoWidth()
+				.Padding(4, 0, 0, 0)
+				[
+					SAssignNew(ViewOptionsComboButton, SComboButton)
+					.ContentPadding(0)
+					.ForegroundColor( this, &SOutputLog::GetViewButtonForegroundColor )
+					.ButtonStyle( FEditorStyle::Get(), "ToggleButton" ) // Use the tool bar item style for this button
+					.OnGetMenuContent( this, &SOutputLog::GetViewButtonContent )
+					.ButtonContent()
+					[
+						SNew(SHorizontalBox)
+ 
+						+SHorizontalBox::Slot()
+						.AutoWidth()
+						.VAlign(VAlign_Center)
+						[
+							SNew(SImage).Image( FEditorStyle::GetBrush("GenericViewButton") )
+						]
+ 
+						+SHorizontalBox::Slot()
+						.AutoWidth()
+						.Padding(2, 0, 0, 0)
+						.VAlign(VAlign_Center)
+						[
+							SNew(STextBlock).Text( LOCTEXT("ViewButton", "View Options") )
+						]
+					]
 				]
 			]
 		]
@@ -1051,6 +1090,14 @@ void SOutputLog::Serialize(const TCHAR* V, ELogVerbosity::Type Verbosity, const
 	}
 }
 
+FSlateColor SOutputLog::GetViewButtonForegroundColor() const
+{
+	static const FName InvertedForegroundName("InvertedForeground");
+	static const FName DefaultForegroundName("DefaultForeground");
+
+	return ViewOptionsComboButton->IsHovered() ? FEditorStyle::GetSlateColor(InvertedForegroundName) : FEditorStyle::GetSlateColor(DefaultForegroundName);
+}
+
 void SOutputLog::ExtendTextBoxMenu(FMenuBuilder& Builder)
 {
 	FUIAction ClearOutputLogAction(
@@ -1095,7 +1142,7 @@ void SOutputLog::RequestForceScroll()
 {
 	if (MessagesTextMarshaller->GetNumFilteredMessages() > 0)
 	{
-		MessagesTextBox->ScrollTo(FTextLocation(MessagesTextMarshaller->GetNumFilteredMessages() - 1));
+		MessagesTextBox->ScrollTo(ETextLocation::EndOfDocument);
 		bIsUserScrolled = false;
 	}
 }
@@ -1111,6 +1158,24 @@ void SOutputLog::Refresh()
 	RequestForceScroll();
 }
 
+bool SOutputLog::IsWordWrapEnabled() const
+{
+	bool WordWrapEnabled = false;
+	GConfig->GetBool(TEXT("/Script/UnrealEd.EditorPerProjectUserSettings"), TEXT("bEnableOutputLogWordWrap"), WordWrapEnabled, GEditorPerProjectIni);
+	return WordWrapEnabled;
+}
+
+void SOutputLog::SetWordWrapEnabled(ECheckBoxState InValue)
+{
+	const bool WordWrapEnabled = (InValue == ECheckBoxState::Checked);
+	GConfig->SetBool(TEXT("/Script/UnrealEd.EditorPerProjectUserSettings"), TEXT("bEnableOutputLogWordWrap"), WordWrapEnabled, GEditorPerProjectIni);
+
+	if (!bIsUserScrolled)
+	{
+		RequestForceScroll();
+	}
+}
+
 void SOutputLog::OnFilterTextChanged(const FText& InFilterText)
 {
 	if (Filter.GetFilterText().ToString().Equals(InFilterText.ToString(), ESearchCase::CaseSensitive))
@@ -1309,6 +1374,28 @@ void SOutputLog::CategoriesSingle_Execute(FName InName)
 	Refresh();
 }
 
+TSharedRef<SWidget> SOutputLog::GetViewButtonContent()
+{
+	TSharedPtr<FExtender> Extender;
+	FMenuBuilder MenuBuilder(true, nullptr, Extender, true);
+	MenuBuilder.AddMenuEntry(
+		LOCTEXT("WordWrapEnabledOption", "Enable Word Wrapping"),
+		LOCTEXT("WordWrapEnabledOptionToolTip", "Enable word wrapping in the Output Log."),
+		FSlateIcon(),
+		FUIAction(
+			FExecuteAction::CreateLambda([this] {
+				// This is a toggle, hence that it is inverted
+				SetWordWrapEnabled(IsWordWrapEnabled() ? ECheckBoxState::Unchecked : ECheckBoxState::Checked);
+			}),
+			FCanExecuteAction::CreateLambda([] { return true; }),
+			FIsActionChecked::CreateSP(this, &SOutputLog::IsWordWrapEnabled)
+		),
+		NAME_None,
+		EUserInterfaceActionType::ToggleButton
+	);
+	return MenuBuilder.MakeWidget();
+}
+
 bool FLogFilter::IsMessageAllowed(const TSharedPtr<FLogMessage>& Message)
 {
 	// Filter Verbosity
diff --git a/Engine/Source/Developer/OutputLog/Private/SOutputLog.h b/Engine/Source/Developer/OutputLog/Private/SOutputLog.h
--- a/Engine/Source/Developer/OutputLog/Private/SOutputLog.h
+++ b/Engine/Source/Developer/OutputLog/Private/SOutputLog.h
@@ -433,9 +433,19 @@ private:
 	/** Forces re-population of the messages list */
 	void Refresh();
 
+	bool IsWordWrapEnabled() const;
+
+	void SetWordWrapEnabled(ECheckBoxState InValue);
+
+	FSlateColor GetViewButtonForegroundColor() const;
+
+	TSharedRef<SWidget> GetViewButtonContent();
+
 public:
 	/** Visible messages filter */
 	FLogFilter Filter;
+
+	TSharedPtr<class SComboButton> ViewOptionsComboButton;
 };
 
 /** Output log text marshaller to convert an array of FLogMessages into styled lines to be consumed by an FTextLayout */
diff --git a/Engine/Source/Runtime/Slate/Private/Widgets/Input/SEditableText.cpp b/Engine/Source/Runtime/Slate/Private/Widgets/Input/SEditableText.cpp
--- a/Engine/Source/Runtime/Slate/Private/Widgets/Input/SEditableText.cpp
+++ b/Engine/Source/Runtime/Slate/Private/Widgets/Input/SEditableText.cpp
@@ -387,9 +387,9 @@ void SEditableText::GoTo(const FTextLocation& NewLocation)
 	EditableTextLayout->GoTo(NewLocation);
 }
 
-void SEditableText::GoTo(ETextLocation GoToLocation)
+void SEditableText::GoTo(const ETextLocation NewLocation)
 {
-	EditableTextLayout->GoTo(GoToLocation);
+	EditableTextLayout->GoTo(NewLocation);
 }
 
 void SEditableText::ScrollTo(const FTextLocation& NewLocation)
@@ -397,6 +397,11 @@ void SEditableText::ScrollTo(const FTextLocation& NewLocation)
 	EditableTextLayout->ScrollTo(NewLocation);
 }
 
+void SEditableText::ScrollTo(const ETextLocation NewLocation)
+{
+	EditableTextLayout->ScrollTo(NewLocation);
+}
+
 void SEditableText::BeginSearch(const FText& InSearchText, const ESearchCase::Type InSearchCase, const bool InReverse)
 {
 	EditableTextLayout->BeginSearch(InSearchText, InSearchCase, InReverse);
diff --git a/Engine/Source/Runtime/Slate/Private/Widgets/Text/SMultiLineEditableText.cpp b/Engine/Source/Runtime/Slate/Private/Widgets/Text/SMultiLineEditableText.cpp
--- a/Engine/Source/Runtime/Slate/Private/Widgets/Text/SMultiLineEditableText.cpp
+++ b/Engine/Source/Runtime/Slate/Private/Widgets/Text/SMultiLineEditableText.cpp
@@ -459,9 +459,9 @@ void SMultiLineEditableText::GoTo(const FTextLocation& NewLocation)
 	EditableTextLayout->GoTo(NewLocation);
 }
 
-void SMultiLineEditableText::GoTo(ETextLocation GoToLocation)
+void SMultiLineEditableText::GoTo(const ETextLocation NewLocation)
 {
-	EditableTextLayout->GoTo(GoToLocation);
+	EditableTextLayout->GoTo(NewLocation);
 }
 
 void SMultiLineEditableText::ScrollTo(const FTextLocation& NewLocation)
@@ -469,6 +469,11 @@ void SMultiLineEditableText::ScrollTo(const FTextLocation& NewLocation)
 	EditableTextLayout->ScrollTo(NewLocation);
 }
 
+void SMultiLineEditableText::ScrollTo(const ETextLocation NewLocation)
+{
+	EditableTextLayout->ScrollTo(NewLocation);
+}
+
 void SMultiLineEditableText::ApplyToSelection(const FRunInfo& InRunInfo, const FTextBlockStyle& InStyle)
 {
 	EditableTextLayout->ApplyToSelection(InRunInfo, InStyle);
diff --git a/Engine/Source/Runtime/Slate/Private/Widgets/Text/SlateEditableTextLayout.cpp b/Engine/Source/Runtime/Slate/Private/Widgets/Text/SlateEditableTextLayout.cpp
--- a/Engine/Source/Runtime/Slate/Private/Widgets/Text/SlateEditableTextLayout.cpp
+++ b/Engine/Source/Runtime/Slate/Private/Widgets/Text/SlateEditableTextLayout.cpp
@@ -2181,7 +2181,7 @@ void FSlateEditableTextLayout::GoTo(const FTextLocation& NewLocation)
 	}
 }
 
-void FSlateEditableTextLayout::GoTo(ETextLocation NewLocation)
+void FSlateEditableTextLayout::GoTo(const ETextLocation NewLocation)
 {
 	JumpTo(NewLocation, ECursorAction::MoveCursor);
 }
@@ -2466,6 +2466,29 @@ void FSlateEditableTextLayout::ScrollTo(const FTextLocation& NewLocation)
 	}
 }
 
+void FSlateEditableTextLayout::ScrollTo(const ETextLocation NewLocation)
+{
+	const TArray< FTextLayout::FLineModel >& Lines = TextLayout->GetLineModels();
+
+	FTextLocation ResolvedTextLocation;
+	switch (NewLocation)
+	{
+	case ETextLocation::BeginningOfDocument:
+		ResolvedTextLocation = FTextLocation(0);
+		break;
+
+	case ETextLocation::EndOfDocument:
+		ResolvedTextLocation = FTextLocation(FMath::Max(0, Lines.Num() - 1));
+		break;
+
+	default:
+		checkf(false, TEXT("Unsupported ETextLocation mode passed to ScrollTo!"));
+		break;
+	}
+
+	ScrollTo(ResolvedTextLocation);
+}
+
 void FSlateEditableTextLayout::UpdateCursorHighlight()
 {
 	PositionToScrollIntoView = SlateEditableTextTypes::FScrollInfo(CursorInfo.GetCursorInteractionLocation(), CursorInfo.GetCursorAlignment());
diff --git a/Engine/Source/Runtime/Slate/Public/Widgets/Input/SEditableText.h b/Engine/Source/Runtime/Slate/Public/Widgets/Input/SEditableText.h
--- a/Engine/Source/Runtime/Slate/Public/Widgets/Input/SEditableText.h
+++ b/Engine/Source/Runtime/Slate/Public/Widgets/Input/SEditableText.h
@@ -317,11 +317,14 @@ public:
 	void GoTo(const FTextLocation& NewLocation);
 
 	/** Move the cursor specified location */
-	void GoTo(ETextLocation NewLocation);
+	void GoTo(const ETextLocation NewLocation);
 
 	/** Scroll to the given location in the document (without moving the cursor) */
 	void ScrollTo(const FTextLocation& NewLocation);
 
+	/** Scroll to the given location in the document (without moving the cursor) */
+	void ScrollTo(const ETextLocation NewLocation);
+
 	/** Begin a new text search (this is called automatically when the bound search text changes) */
 	void BeginSearch(const FText& InSearchText, const ESearchCase::Type InSearchCase = ESearchCase::IgnoreCase, const bool InReverse = false);
 
diff --git a/Engine/Source/Runtime/Slate/Public/Widgets/Input/SEditableTextBox.h b/Engine/Source/Runtime/Slate/Public/Widgets/Input/SEditableTextBox.h
--- a/Engine/Source/Runtime/Slate/Public/Widgets/Input/SEditableTextBox.h
+++ b/Engine/Source/Runtime/Slate/Public/Widgets/Input/SEditableTextBox.h
@@ -320,7 +320,7 @@ public:
 	void GoTo(const FTextLocation& NewLocation);
 
 	/** Move the cursor to the specified location */
-	void GoTo(ETextLocation NewLocation)
+	void GoTo(const ETextLocation NewLocation)
 	{
 		EditableText->GoTo(NewLocation);
 	}
@@ -328,6 +328,12 @@ public:
 	/** Scroll to the given location in the document (without moving the cursor) */
 	void ScrollTo(const FTextLocation& NewLocation);
 
+	/** Scroll to the given location in the document (without moving the cursor) */
+	void ScrollTo(const ETextLocation NewLocation)
+	{
+		EditableText->GoTo(NewLocation);
+	}
+
 	/** Begin a new text search (this is called automatically when the bound search text changes) */
 	void BeginSearch(const FText& InSearchText, const ESearchCase::Type InSearchCase = ESearchCase::IgnoreCase, const bool InReverse = false);
 
diff --git a/Engine/Source/Runtime/Slate/Public/Widgets/Input/SMultiLineEditableTextBox.h b/Engine/Source/Runtime/Slate/Public/Widgets/Input/SMultiLineEditableTextBox.h
--- a/Engine/Source/Runtime/Slate/Public/Widgets/Input/SMultiLineEditableTextBox.h
+++ b/Engine/Source/Runtime/Slate/Public/Widgets/Input/SMultiLineEditableTextBox.h
@@ -382,7 +382,7 @@ public:
 	void GoTo(const FTextLocation& NewLocation);
 
 	/** Move the cursor to the specified location */
-	void GoTo(ETextLocation NewLocation)
+	void GoTo(const ETextLocation NewLocation)
 	{
 		EditableText->GoTo(NewLocation);
 	}
@@ -390,6 +390,12 @@ public:
 	/** Scroll to the given location in the document (without moving the cursor) */
 	void ScrollTo(const FTextLocation& NewLocation);
 
+	/** Scroll to the given location in the document (without moving the cursor) */
+	void ScrollTo(const ETextLocation NewLocation)
+	{
+		EditableText->ScrollTo(NewLocation);
+	}
+
 	/** Apply the given style to the currently selected text (or insert a new run at the current cursor position if no text is selected) */
 	void ApplyToSelection(const FRunInfo& InRunInfo, const FTextBlockStyle& InStyle);
 
diff --git a/Engine/Source/Runtime/Slate/Public/Widgets/Text/SMultiLineEditableText.h b/Engine/Source/Runtime/Slate/Public/Widgets/Text/SMultiLineEditableText.h
--- a/Engine/Source/Runtime/Slate/Public/Widgets/Text/SMultiLineEditableText.h
+++ b/Engine/Source/Runtime/Slate/Public/Widgets/Text/SMultiLineEditableText.h
@@ -328,11 +328,14 @@ public:
 	void GoTo(const FTextLocation& NewLocation);
 
 	/** Move the cursor specified location */
-	void GoTo(ETextLocation NewLocation);
+	void GoTo(const ETextLocation NewLocation);
 
 	/** Scroll to the given location in the document (without moving the cursor) */
 	void ScrollTo(const FTextLocation& NewLocation);
 
+	/** Scroll to the given location in the document (without moving the cursor) */
+	void ScrollTo(const ETextLocation NewLocation);
+
 	/** Apply the given style to the currently selected text (or insert a new run at the current cursor position if no text is selected) */
 	void ApplyToSelection(const FRunInfo& InRunInfo, const FTextBlockStyle& InStyle);
 
diff --git a/Engine/Source/Runtime/Slate/Public/Widgets/Text/SlateEditableTextLayout.h b/Engine/Source/Runtime/Slate/Public/Widgets/Text/SlateEditableTextLayout.h
--- a/Engine/Source/Runtime/Slate/Public/Widgets/Text/SlateEditableTextLayout.h
+++ b/Engine/Source/Runtime/Slate/Public/Widgets/Text/SlateEditableTextLayout.h
@@ -241,7 +241,7 @@ public:
 	void GoTo(const FTextLocation& NewLocation);
 
 	/** Move the cursor specified location */
-	void GoTo(ETextLocation NewLocation);
+	void GoTo(const ETextLocation NewLocation);
 
 	/** Jump the cursor to the given location in the document */
 	void JumpTo(ETextLocation JumpLocation, ECursorAction Action);
@@ -249,6 +249,9 @@ public:
 	/** Scroll to the given location in the document (without moving the cursor) */
 	void ScrollTo(const FTextLocation& NewLocation);
 
+	/** Scroll to the given location in the document (without moving the cursor) */
+	void ScrollTo(const ETextLocation NewLocation);
+
 	/** Update the active cursor highlight based on the state of the text layout */
 	void UpdateCursorHighlight();
 
