Added Clear on PIE to Output Log options, along with Open Source Location of the log file and Open in External Editor of log file.

diff --git a/Engine/Config/BaseEditorPerProjectUserSettings.ini b/Engine/Config/BaseEditorPerProjectUserSettings.ini
--- a/Engine/Config/BaseEditorPerProjectUserSettings.ini
+++ b/Engine/Config/BaseEditorPerProjectUserSettings.ini
@@ -20,6 +20,7 @@
 
 [/Script/UnrealEd.EditorPerProjectUserSettings]
 bEnableOutputLogWordWrap=False
+bEnableOutputLogClearOnPIE=False
 
 ; True if WASD keys should be remapped to flight controls while the right mouse button is held down
 AllowFlightCameraToRemapKeys=True
diff --git a/Engine/Source/Developer/OutputLog/OutputLog.Build.cs b/Engine/Source/Developer/OutputLog/OutputLog.Build.cs
--- a/Engine/Source/Developer/OutputLog/OutputLog.Build.cs
+++ b/Engine/Source/Developer/OutputLog/OutputLog.Build.cs
@@ -18,5 +18,14 @@ public class OutputLog : ModuleRules
                 "DesktopPlatform"
 			}
 		);
-	}
+
+		if (Target.bBuildEditor)
+        {
+            PrivateDependencyModuleNames.AddRange(
+				new string[] {
+                        "UnrealEd"
+				}
+			);
+        }
+    }
 }
diff --git a/Engine/Source/Developer/OutputLog/Private/OutputLogModule.cpp b/Engine/Source/Developer/OutputLog/Private/OutputLogModule.cpp
--- a/Engine/Source/Developer/OutputLog/Private/OutputLogModule.cpp
+++ b/Engine/Source/Developer/OutputLog/Private/OutputLogModule.cpp
@@ -65,6 +65,9 @@ private:
 /** Our global output log app spawner */
 static TSharedPtr<FOutputLogHistory> OutputLogHistory;
 
+/** Our global active output log */
+static TWeakPtr<SOutputLog> OutputLog;
+
 TSharedRef<SDockTab> SpawnOutputLog( const FSpawnTabArgs& Args )
 {
 	return SNew(SDockTab)
@@ -72,7 +75,7 @@ TSharedRef<SDockTab> SpawnOutputLog( const FSpawnTabArgs& Args )
 		.TabRole( ETabRole::NomadTab )
 		.Label( NSLOCTEXT("OutputLog", "TabTitle", "Output Log") )
 		[
-			SNew(SOutputLog).Messages( OutputLogHistory->GetMessages() )
+			SAssignNew(OutputLog, SOutputLog).Messages( OutputLogHistory->GetMessages() )
 		];
 }
 
@@ -100,7 +103,9 @@ void FOutputLogModule::StartupModule()
 		.SetTooltipText(NSLOCTEXT("UnrealEditor", "DeviceOutputLogTooltipText", "Open the Device Output Log tab."))
 		.SetGroup( WorkspaceMenu::GetMenuStructure().GetDeveloperToolsLogCategory() )
 		.SetIcon( FSlateIcon(FEditorStyle::GetStyleSetName(), "Log.TabIcon") );
-	
+
+	FEditorDelegates::BeginPIE.AddRaw(this, &FOutputLogModule::ClearOnPIE);
+
 	OutputLogHistory = MakeShareable(new FOutputLogHistory);
 }
 
@@ -112,6 +117,8 @@ void FOutputLogModule::ShutdownModule()
 		FGlobalTabmanager::Get()->UnregisterNomadTabSpawner(OutputLogModule::DeviceOutputLogTabName);
 	}
 
+	FEditorDelegates::BeginPIE.RemoveAll(this);
+
 	OutputLogHistory.Reset();
 }
 
@@ -186,3 +193,22 @@ void FOutputLogModule::CloseDebugConsole()
 		}
 	}
 }
+
+void FOutputLogModule::ClearOnPIE(const bool bIsSimulating)
+{
+	if (OutputLog.IsValid())
+	{
+		bool bClearOnPIEEnabled = false;
+		GConfig->GetBool(TEXT("/Script/UnrealEd.EditorPerProjectUserSettings"), TEXT("bEnableOutputLogClearOnPIE"), bClearOnPIEEnabled, GEditorPerProjectIni);
+
+		if (bClearOnPIEEnabled)
+		{
+			TSharedPtr<SOutputLog> OutputLogShared = OutputLog.Pin();
+
+			if (OutputLogShared->CanClearLog())
+			{
+				OutputLogShared->OnClearLog();
+			}
+		}
+	}
+}
diff --git a/Engine/Source/Developer/OutputLog/Private/SOutputLog.cpp b/Engine/Source/Developer/OutputLog/Private/SOutputLog.cpp
--- a/Engine/Source/Developer/OutputLog/Private/SOutputLog.cpp
+++ b/Engine/Source/Developer/OutputLog/Private/SOutputLog.cpp
@@ -24,6 +24,8 @@
 #include "Widgets/Images/SImage.h"
 #include "Features/IModularFeatures.h"
 #include "Misc/CoreDelegates.h"
+#include "HAL/PlatformOutputDevices.h"
+
 
 #define LOCTEXT_NAMESPACE "SOutputLog"
 /** Expression context to test the given messages against the current text filter */
@@ -1176,6 +1178,19 @@ void SOutputLog::SetWordWrapEnabled(ECheckBoxState InValue)
 	}
 }
 
+bool SOutputLog::IsClearOnPIEEnabled() const
+{
+	bool ClearOnPIEEnabled = false;
+	GConfig->GetBool(TEXT("/Script/UnrealEd.EditorPerProjectUserSettings"), TEXT("bEnableOutputLogClearOnPIE"), ClearOnPIEEnabled, GEditorPerProjectIni);
+	return ClearOnPIEEnabled;
+}
+
+void SOutputLog::SetClearOnPIE(ECheckBoxState InValue)
+{
+	const bool ClearOnPIEEnabled = (InValue == ECheckBoxState::Checked);
+	GConfig->SetBool(TEXT("/Script/UnrealEd.EditorPerProjectUserSettings"), TEXT("bEnableOutputLogClearOnPIE"), ClearOnPIEEnabled, GEditorPerProjectIni);
+}
+
 void SOutputLog::OnFilterTextChanged(const FText& InFilterText)
 {
 	if (Filter.GetFilterText().ToString().Equals(InFilterText.ToString(), ESearchCase::CaseSensitive))
@@ -1393,9 +1408,69 @@ TSharedRef<SWidget> SOutputLog::GetViewButtonContent()
 		NAME_None,
 		EUserInterfaceActionType::ToggleButton
 	);
+	MenuBuilder.AddMenuEntry(
+		LOCTEXT("ClearOnPIE", "Clear on PIE"),
+		LOCTEXT("ClearOnPIEToolTip", "Enable clearing of the Output Log on PIE startup."),
+		FSlateIcon(),
+		FUIAction(
+			FExecuteAction::CreateLambda([this] {
+				// This is a toggle, hence that it is inverted
+				SetClearOnPIE(IsClearOnPIEEnabled() ? ECheckBoxState::Unchecked : ECheckBoxState::Checked);
+			}),
+			FCanExecuteAction::CreateLambda([] { return true; }),
+			FIsActionChecked::CreateSP(this, &SOutputLog::IsClearOnPIEEnabled)
+		),
+		NAME_None,
+		EUserInterfaceActionType::ToggleButton
+	);
+	MenuBuilder.AddMenuSeparator();
+
+	//Show Source In Explorer
+	MenuBuilder.AddMenuEntry(
+		LOCTEXT("FindSourceFile", "Open Source Location"),
+		LOCTEXT("FindSourceFileTooltip", "Opens the folder containing the source of the Output Log."),
+		FSlateIcon(FEditorStyle::GetStyleSetName(), "OutputLog.OpenSourceLocation"),
+		FUIAction(
+			FExecuteAction::CreateSP(this, &SOutputLog::OpenLogFileInExplorer)
+		)
+	);
+	
+	// Open In External Editor
+	MenuBuilder.AddMenuEntry(
+		LOCTEXT("OpenInExternalEditor", "Open In External Editor"),
+		LOCTEXT("OpenInExternalEditorTooltip", "Opens the Output Log in the default external editor."),
+		FSlateIcon(FEditorStyle::GetStyleSetName(), "OutputLog.OpenInExternalEditor"),
+		FUIAction(
+			FExecuteAction::CreateSP(this, &SOutputLog::OpenLogFileInExternalEditor)
+		)
+	);
+
+
 	return MenuBuilder.MakeWidget();
 }
 
+void SOutputLog::OpenLogFileInExplorer()
+{
+	FString Path = FPaths::ConvertRelativePathToFull(FPaths::ProjectLogDir());
+	if (!Path.Len() || !IFileManager::Get().DirectoryExists(*Path))
+	{
+		return;
+	}
+
+	FPlatformProcess::ExploreFolder(*FPaths::GetPath(Path));
+}
+
+void SOutputLog::OpenLogFileInExternalEditor()
+{
+	FString Path = FPaths::ConvertRelativePathToFull(FGenericPlatformOutputDevices::GetAbsoluteLogFilename());
+	if (!Path.Len() || IFileManager::Get().FileSize(*Path) == INDEX_NONE)
+	{
+		return;
+	}
+
+	FPlatformProcess::LaunchFileInDefaultExternalApplication(*Path, NULL, ELaunchVerb::Open);
+} 
+
 bool FLogFilter::IsMessageAllowed(const TSharedPtr<FLogMessage>& Message)
 {
 	// Filter Verbosity
diff --git a/Engine/Source/Developer/OutputLog/Private/SOutputLog.h b/Engine/Source/Developer/OutputLog/Private/SOutputLog.h
--- a/Engine/Source/Developer/OutputLog/Private/SOutputLog.h
+++ b/Engine/Source/Developer/OutputLog/Private/SOutputLog.h
@@ -341,6 +341,16 @@ public:
 	 */
 	static bool CreateLogMessages(const TCHAR* V, ELogVerbosity::Type Verbosity, const class FName& Category, TArray< TSharedPtr<FLogMessage> >& OutMessages);
 
+	/**
+	* Called when delete all is selected
+	*/
+	void OnClearLog();
+
+	/**
+	 * Called to determine whether delete all is currently a valid command
+	 */
+	bool CanClearLog() const;
+
 protected:
 
 	virtual void Serialize( const TCHAR* V, ELogVerbosity::Type Verbosity, const class FName& Category ) override;
@@ -353,22 +363,12 @@ protected:
 	 * Extends the context menu used by the text box
 	 */
 	void ExtendTextBoxMenu(FMenuBuilder& Builder);
-
-	/**
-	 * Called when delete all is selected
-	 */
-	void OnClearLog();
-
+	
 	/**
 	 * Called when the user scrolls the log window vertically
 	 */
 	void OnUserScrolled(float ScrollOffset);
 
-	/**
-	 * Called to determine whether delete all is currently a valid command
-	 */
-	bool CanClearLog() const;
-
 	/** Called when a console command is entered for this output log */
 	void OnConsoleCommandExecuted();
 
@@ -437,10 +437,18 @@ private:
 
 	void SetWordWrapEnabled(ECheckBoxState InValue);
 
+	bool IsClearOnPIEEnabled() const;
+
+	void SetClearOnPIE(ECheckBoxState InValue);
+
 	FSlateColor GetViewButtonForegroundColor() const;
 
 	TSharedRef<SWidget> GetViewButtonContent();
 
+	void OpenLogFileInExplorer();
+
+	void OpenLogFileInExternalEditor();
+
 public:
 	/** Visible messages filter */
 	FLogFilter Filter;
diff --git a/Engine/Source/Developer/OutputLog/Public/OutputLogModule.h b/Engine/Source/Developer/OutputLog/Public/OutputLogModule.h
--- a/Engine/Source/Developer/OutputLog/Public/OutputLogModule.h
+++ b/Engine/Source/Developer/OutputLog/Public/OutputLogModule.h
@@ -47,6 +47,8 @@ public:
 	/** Closes the debug console for the specified window */
 	virtual void CloseDebugConsole();
 
+	virtual void ClearOnPIE(const bool bIsSimulating);
+
 private:
 
 	/** Weak pointer to a debug console that's currently open, if any */
diff --git a/Engine/Source/Editor/EditorStyle/Private/SlateEditorStyle.cpp b/Engine/Source/Editor/EditorStyle/Private/SlateEditorStyle.cpp
--- a/Engine/Source/Editor/EditorStyle/Private/SlateEditorStyle.cpp
+++ b/Engine/Source/Editor/EditorStyle/Private/SlateEditorStyle.cpp
@@ -1207,6 +1207,10 @@ void FSlateEditorStyle::FStyle::SetupGeneralStyles()
 			);
 
 		Set("DebugConsole.Background", new BOX_BRUSH("Old/Menu_Background", FMargin(8.0f / 64.0f)));
+
+		Set("OutputLog.OpenSourceLocation", new IMAGE_BRUSH("Icons/icon_Asset_Open_Source_Location_16x", Icon16x16));
+		Set("OutputLog.OpenInExternalEditor", new IMAGE_BRUSH("Icons/icon_Asset_Open_In_External_Editor_16x", Icon16x16));
+
 	}
 
 	// Debug Tools Window
