Added necessary header files.

diff --git a/Engine/Source/Developer/OutputLog/Private/OutputLogModule.cpp b/Engine/Source/Developer/OutputLog/Private/OutputLogModule.cpp
--- a/Engine/Source/Developer/OutputLog/Private/OutputLogModule.cpp
+++ b/Engine/Source/Developer/OutputLog/Private/OutputLogModule.cpp
@@ -13,6 +13,7 @@
 #include "Editor/WorkspaceMenuStructure/Public/WorkspaceMenuStructure.h"
 #include "Editor/WorkspaceMenuStructure/Public/WorkspaceMenuStructureModule.h"
 #include "Widgets/Docking/SDockTab.h"
+#include "Editor.h"
 
 IMPLEMENT_MODULE( FOutputLogModule, OutputLog );
 
diff --git a/Engine/Source/Developer/OutputLog/Private/SOutputLog.cpp b/Engine/Source/Developer/OutputLog/Private/SOutputLog.cpp
--- a/Engine/Source/Developer/OutputLog/Private/SOutputLog.cpp
+++ b/Engine/Source/Developer/OutputLog/Private/SOutputLog.cpp
@@ -25,7 +25,7 @@
 #include "Features/IModularFeatures.h"
 #include "Misc/CoreDelegates.h"
 #include "HAL/PlatformOutputDevices.h"
-
+#include "HAL/FileManager.h"
 
 #define LOCTEXT_NAMESPACE "SOutputLog"
 /** Expression context to test the given messages against the current text filter */
