Added forgotten header for configs.

diff --git a/Engine/Source/Developer/OutputLog/Private/OutputLogModule.cpp b/Engine/Source/Developer/OutputLog/Private/OutputLogModule.cpp
--- a/Engine/Source/Developer/OutputLog/Private/OutputLogModule.cpp
+++ b/Engine/Source/Developer/OutputLog/Private/OutputLogModule.cpp
@@ -14,6 +14,7 @@
 #include "Editor/WorkspaceMenuStructure/Public/WorkspaceMenuStructureModule.h"
 #include "Widgets/Docking/SDockTab.h"
 #include "Editor.h"
+#include "Misc/ConfigCacheIni.h"
 
 IMPLEMENT_MODULE( FOutputLogModule, OutputLog );
 
