Fix Tools Win64 CIS for ensuring it doesn't access editor build functionality.

diff --git a/Engine/Source/Developer/OutputLog/Private/OutputLogModule.cpp b/Engine/Source/Developer/OutputLog/Private/OutputLogModule.cpp
--- a/Engine/Source/Developer/OutputLog/Private/OutputLogModule.cpp
+++ b/Engine/Source/Developer/OutputLog/Private/OutputLogModule.cpp
@@ -13,9 +13,12 @@
 #include "Editor/WorkspaceMenuStructure/Public/WorkspaceMenuStructure.h"
 #include "Editor/WorkspaceMenuStructure/Public/WorkspaceMenuStructureModule.h"
 #include "Widgets/Docking/SDockTab.h"
-#include "Editor.h"
 #include "Misc/ConfigCacheIni.h"
 
+#if WITH_EDITOR
+#include "Editor.h"
+#endif
+
 IMPLEMENT_MODULE( FOutputLogModule, OutputLog );
 
 namespace OutputLogModule
@@ -106,7 +109,9 @@ void FOutputLogModule::StartupModule()
 		.SetGroup( WorkspaceMenu::GetMenuStructure().GetDeveloperToolsLogCategory() )
 		.SetIcon( FSlateIcon(FEditorStyle::GetStyleSetName(), "Log.TabIcon") );
 
+#if WITH_EDITOR
 	FEditorDelegates::BeginPIE.AddRaw(this, &FOutputLogModule::ClearOnPIE);
+#endif
 
 	OutputLogHistory = MakeShareable(new FOutputLogHistory);
 }
@@ -119,7 +124,9 @@ void FOutputLogModule::ShutdownModule()
 		FGlobalTabmanager::Get()->UnregisterNomadTabSpawner(OutputLogModule::DeviceOutputLogTabName);
 	}
 
+#if WITH_EDITOR
 	FEditorDelegates::BeginPIE.RemoveAll(this);
+#endif
 
 	OutputLogHistory.Reset();
 }
