Do not allow the console command window to absorb and enter a key when the content browser drawer key is pressed

diff --git a/Engine/Source/Developer/OutputLog/Private/SOutputLog.cpp b/Engine/Source/Developer/OutputLog/Private/SOutputLog.cpp
--- a/Engine/Source/Developer/OutputLog/Private/SOutputLog.cpp
+++ b/Engine/Source/Developer/OutputLog/Private/SOutputLog.cpp
@@ -627,6 +627,12 @@ FReply SConsoleInputBox::OnKeyCharHandler(const FGeometry& MyGeometry, const FCh
 		return FReply::Handled();
 	}
 
+	if (InCharacterEvent.GetModifierKeys().AnyModifiersDown() && InCharacterEvent.GetCharacter() == ' ')
+	{	
+		// Ignore space bar + a modifier key.  It should not type a space as this is used by other keyboard shortcuts
+		return FReply::Handled();
+	}
+
 	FInputChord OpenConsoleChord;
 	if (ActiveCommandExecutor && ActiveCommandExecutor->AllowHotKeyClose())
 	{
diff --git a/Engine/Source/Runtime/ApplicationCore/Public/GenericPlatform/GenericApplication.h b/Engine/Source/Runtime/ApplicationCore/Public/GenericPlatform/GenericApplication.h
--- a/Engine/Source/Runtime/ApplicationCore/Public/GenericPlatform/GenericApplication.h
+++ b/Engine/Source/Runtime/ApplicationCore/Public/GenericPlatform/GenericApplication.h
@@ -277,6 +277,14 @@ public:
 		return AllModifersDown;
 	}
 
+	/**
+	 * @return true if any modifier key is down
+	 */
+	bool AnyModifiersDown() const
+	{
+		return IsControlDown() || IsCommandDown() || IsShiftDown() || IsAltDown();
+	}
+
 private:
 
 	/** True if the left shift key was down when this event occurred. */
