Fix console not appearing in immersive mode

diff --git a/Engine/Source/Developer/OutputLog/Private/OutputLogModule.cpp b/Engine/Source/Developer/OutputLog/Private/OutputLogModule.cpp
--- a/Engine/Source/Developer/OutputLog/Private/OutputLogModule.cpp
+++ b/Engine/Source/Developer/OutputLog/Private/OutputLogModule.cpp
@@ -153,6 +153,12 @@ void FOutputLogModule::ToggleDebugConsoleForWindow(const TSharedRef<SWindow>& Wi
 		TSharedPtr< SWindow > WindowForExistingConsole = FSlateApplication::Get().FindWidgetWindow(PinnedDebugConsole.ToSharedRef());
 		if (WindowForExistingConsole.IsValid())
 		{
+			if (PreviousKeyboardFocusedWidget.IsValid())
+			{
+				FSlateApplication::Get().SetKeyboardFocus(PreviousKeyboardFocusedWidget.Pin());
+				PreviousKeyboardFocusedWidget.Reset();
+			}
+
 			WindowForExistingConsole->RemoveOverlaySlot(PinnedDebugConsole.ToSharedRef());
 			DebugConsole.Reset();
 		}
@@ -186,6 +192,8 @@ void FOutputLogModule::ToggleDebugConsoleForWindow(const TSharedRef<SWindow>& Wi
 				DebugConsoleRef
 			];
 
+		PreviousKeyboardFocusedWidget = FSlateApplication::Get().GetKeyboardFocusedWidget();
+	
 		// Force keyboard focus
 		DebugConsoleRef->SetFocusToEditableText();
 	}
diff --git a/Engine/Source/Developer/OutputLog/Private/SDebugConsole.cpp b/Engine/Source/Developer/OutputLog/Private/SDebugConsole.cpp
--- a/Engine/Source/Developer/OutputLog/Private/SDebugConsole.cpp
+++ b/Engine/Source/Developer/OutputLog/Private/SDebugConsole.cpp
@@ -62,8 +62,8 @@ void SDebugConsole::Construct( const FArguments& InArgs, const EDebugConsoleStyl
 			[
 				SNew( SBorder )
 				.VAlign( VAlign_Center )
-				.Padding( FMargin(5.0f, 2.0f) )
-				.BorderImage( FEditorStyle::GetBrush( "DebugConsole.Background" ) )
+				.Padding( FMargin(3.0f, 3.0f) )
+				.BorderImage(FAppStyle::Get().GetBrush("Brushes.Background"))
 				.ColorAndOpacity( this, &SDebugConsole::GetAnimatedColorAndOpacity )
 				.BorderBackgroundColor( this, &SDebugConsole::GetAnimatedSlateColor )
 				[
diff --git a/Engine/Source/Developer/OutputLog/Public/OutputLogModule.h b/Engine/Source/Developer/OutputLog/Public/OutputLogModule.h
--- a/Engine/Source/Developer/OutputLog/Public/OutputLogModule.h
+++ b/Engine/Source/Developer/OutputLog/Public/OutputLogModule.h
@@ -50,7 +50,9 @@ public:
 	virtual void ClearOnPIE(const bool bIsSimulating);
 
 private:
-
 	/** Weak pointer to a debug console that's currently open, if any */
-	TWeakPtr< SWidget > DebugConsole;
+	TWeakPtr<SWidget> DebugConsole;
+
+	/** Weak pointer to the widget to focus once they console window closes */
+	TWeakPtr<SWidget> PreviousKeyboardFocusedWidget;
 };
diff --git a/Engine/Source/Editor/StatusBar/Private/SStatusBar.cpp b/Engine/Source/Editor/StatusBar/Private/SStatusBar.cpp
--- a/Engine/Source/Editor/StatusBar/Private/SStatusBar.cpp
+++ b/Engine/Source/Editor/StatusBar/Private/SStatusBar.cpp
@@ -386,9 +386,9 @@ TSharedPtr<SDockTab> SStatusBar::GetParentTab() const
 	return ParentTab.Pin();
 }
 
-void SStatusBar::FocusDebugConsole()
+bool SStatusBar::FocusDebugConsole()
 {
-	FSlateApplication::Get().SetKeyboardFocus(ConsoleEditBox, EFocusCause::SetDirectly);
+	return FSlateApplication::Get().SetKeyboardFocus(ConsoleEditBox, EFocusCause::SetDirectly);
 }
 
 bool SStatusBar::IsDebugConsoleFocused() const
diff --git a/Engine/Source/Editor/StatusBar/Private/StatusBarSubsystem.cpp b/Engine/Source/Editor/StatusBar/Private/StatusBarSubsystem.cpp
--- a/Engine/Source/Editor/StatusBar/Private/StatusBarSubsystem.cpp
+++ b/Engine/Source/Editor/StatusBar/Private/StatusBarSubsystem.cpp
@@ -43,9 +43,7 @@ bool UStatusBarSubsystem::FocusDebugConsole(TSharedRef<SWindow> ParentWindow)
 			{
 				// Cache off the previously focused widget so we can restore focus if the user hits the focus key again
 				PreviousKeyboardFocusedWidget = FSlateApplication::Get().GetKeyboardFocusedWidget();
-				StatusBarPinned->FocusDebugConsole();
-
-				bFocusedSuccessfully = true;
+				bFocusedSuccessfully = StatusBarPinned->FocusDebugConsole();
 				break;
 			}
 		}
diff --git a/Engine/Source/Editor/StatusBar/Public/SStatusBar.h b/Engine/Source/Editor/StatusBar/Public/SStatusBar.h
--- a/Engine/Source/Editor/StatusBar/Public/SStatusBar.h
+++ b/Engine/Source/Editor/StatusBar/Public/SStatusBar.h
@@ -65,7 +65,7 @@ public:
 
 
 	TSharedPtr<SDockTab> GetParentTab() const;
-	void FocusDebugConsole();
+	bool FocusDebugConsole();
 	bool IsDebugConsoleFocused() const;
 
 	bool IsContentBrowserOpened() const;
