Editor: Clicking a suggestion in the auto-complete list now moves the caret to the end of the selected text like using Tab to select the value.

diff --git a/Engine/Source/Developer/OutputLog/Private/SOutputLog.cpp b/Engine/Source/Developer/OutputLog/Private/SOutputLog.cpp
--- a/Engine/Source/Developer/OutputLog/Private/SOutputLog.cpp
+++ b/Engine/Source/Developer/OutputLog/Private/SOutputLog.cpp
@@ -191,6 +191,11 @@ void SConsoleInputBox::SuggestionSelectionChanged(TSharedPtr<FString> NewValue,
 	if( SelectInfo == ESelectInfo::OnMouseClick )
 	{
 		SuggestionBox->SetIsOpen( false );
+
+		// Jump the caret to the end of the newly auto-completed line. This makes it so that selecting
+		// an option doesn't leave the cursor in the middle of the suggestion (which makes it hard to 
+		// ctrl-back out, or to type "?" for help, etc.)
+		InputText->GoTo(ETextLocation::EndOfDocument);
 	}
 
 	// Ideally this would set the focus back to the edit control
