Tools: Output log category improvements
- You can now choose to color-code entire log lines by category, or just the category itself (CategoryColorizationMode in editor appearance settings)
- You can also choose to highlight all messages for a particular category by right-clicking on a line
- Also includes plumbing to expose GetCursorLocation up to SMultiLineEditableTextBox
- Known issues: If you edit editor appearance preferences after scrolling (but not clicking) in the output log, the view in the output log will be snapped to the last clicked cursor location

diff --git a/Engine/Source/Developer/OutputLog/Private/SOutputLog.cpp b/Engine/Source/Developer/OutputLog/Private/SOutputLog.cpp
--- a/Engine/Source/Developer/OutputLog/Private/SOutputLog.cpp
+++ b/Engine/Source/Developer/OutputLog/Private/SOutputLog.cpp
@@ -30,8 +30,94 @@
 #include "Framework/Docking/TabManager.h"
 #include "Widgets/Docking/SDockTab.h"
 #include "OutputLogModule.h"
+#include "Widgets/Text/SlateEditableTextTypes.h"
 
 #define LOCTEXT_NAMESPACE "SOutputLog"
+
+class FCategoryLineHighlighter : public ISlateLineHighlighter
+{
+public:
+	static TSharedRef<FCategoryLineHighlighter> Create()
+	{
+		return MakeShareable(new FCategoryLineHighlighter());
+	}
+
+	virtual int32 OnPaint(const FPaintArgs& Args, const FTextLayout::FLineView& Line, const float OffsetX, const float Width, const FTextBlockStyle& DefaultStyle, const FGeometry& AllottedGeometry, const FSlateRect& MyCullingRect, FSlateWindowElementList& OutDrawElements, int32 LayerId, const FWidgetStyle& InWidgetStyle, bool bParentEnabled) const override
+	{
+		const FVector2D Location(Line.Offset.X + OffsetX, Line.Offset.Y);
+
+		// If we've not been set to an explicit color, calculate a suitable one from the linked color
+		FLinearColor SelectionBackgroundColorAndOpacity = DefaultStyle.SelectedBackgroundColor.GetColor(InWidgetStyle);// *InWidgetStyle.GetColorAndOpacityTint();
+		SelectionBackgroundColorAndOpacity.A *= 0.2f;
+
+		// The block size and offset values are pre-scaled, so we need to account for that when converting the block offsets into paint geometry
+		const float InverseScale = Inverse(AllottedGeometry.Scale);
+
+		if (Width > 0.0f)
+		{
+			// Draw the actual highlight rectangle
+			FSlateDrawElement::MakeBox(
+				OutDrawElements,
+				++LayerId,
+				AllottedGeometry.ToPaintGeometry(TransformVector(InverseScale, FVector2D(Width, FMath::Max(Line.Size.Y, Line.TextHeight))), FSlateLayoutTransform(TransformPoint(InverseScale, Location))),
+				&DefaultStyle.HighlightShape,
+				bParentEnabled /*&& bHasKeyboardFocus*/ ? ESlateDrawEffect::None : ESlateDrawEffect::DisabledEffect,
+				SelectionBackgroundColorAndOpacity
+			);
+		}
+
+		return LayerId;
+	}
+
+protected:
+	FCategoryLineHighlighter()
+	{
+	}
+};
+
+
+class FCategoryBadgeHighlighter : public ISlateLineHighlighter
+{
+public:
+	static TSharedRef<FCategoryBadgeHighlighter> Create(const FLinearColor& InBadgeColor)
+	{
+		return MakeShareable(new FCategoryBadgeHighlighter(InBadgeColor));
+	}
+
+	virtual int32 OnPaint(const FPaintArgs& Args, const FTextLayout::FLineView& Line, const float OffsetX, const float Width, const FTextBlockStyle& DefaultStyle, const FGeometry& AllottedGeometry, const FSlateRect& MyCullingRect, FSlateWindowElementList& OutDrawElements, int32 LayerId, const FWidgetStyle& InWidgetStyle, bool bParentEnabled) const override
+	{
+		const FVector2D Location(Line.Offset.X + OffsetX, Line.Offset.Y);
+
+		// The block size and offset values are pre-scaled, so we need to account for that when converting the block offsets into paint geometry
+		const float InverseScale = Inverse(AllottedGeometry.Scale);
+
+		if (Width > 0.0f)
+		{
+			// Draw the actual highlight rectangle
+			FSlateDrawElement::MakeBox(
+				OutDrawElements,
+				++LayerId,
+				AllottedGeometry.ToPaintGeometry(TransformVector(InverseScale, FVector2D(Width, FMath::Max(Line.Size.Y, Line.TextHeight))), FSlateLayoutTransform(TransformPoint(InverseScale, Location))),
+				&DefaultStyle.HighlightShape,
+				bParentEnabled ? ESlateDrawEffect::None : ESlateDrawEffect::DisabledEffect,
+				BadgeColor
+			);
+		}
+
+		return LayerId;
+	}
+
+protected:
+	FLinearColor BadgeColor;
+
+	FCategoryBadgeHighlighter(const FLinearColor& InBadgeColor)
+		: BadgeColor(InBadgeColor)
+	{
+	}
+};
+
+
+
 /** Expression context to test the given messages against the current text filter */
 class FLogFilter_TextFilterExpressionContextOutputLog : public ITextFilterExpressionContext
 {
@@ -714,6 +800,21 @@ bool FOutputLogTextLayoutMarshaller::SubmitPendingMessages()
 	return false;
 }
 
+float FOutputLogTextLayoutMarshaller::GetCategoryHue(FName CategoryName)
+{
+	if (float* pResult = CategoryHueMap.Find(CategoryName))
+	{
+		return *pResult;
+	}
+	else
+	{
+		FRandomStream RNG(GetTypeHash(CategoryName));
+		const float Hue = RNG.FRandRange(0.0f, 360.0f);
+		CategoryHueMap.Add(CategoryName, Hue);
+		return Hue;
+	}
+}
+
 void FOutputLogTextLayoutMarshaller::AppendPendingMessagesToTextLayout()
 {
 	const int32 CurrentMessagesCount = Messages.Num();
@@ -740,14 +841,29 @@ void FOutputLogTextLayoutMarshaller::AppendPendingMessagesToTextLayout()
 		MakeDirty();
 	}
 
+	const ELogCategoryColorizationMode CategoryColorizationMode = GetDefault<UEditorStyleSettings>()->CategoryColorizationMode;
+
 	TArray<FTextLayout::FNewLineData> LinesToAdd;
 	LinesToAdd.Reserve(NumPendingMessages);
+	TArray<FTextLineHighlight> Highlights;
 
 	int32 NumAddedMessages = 0;
 
+	auto ComputeCategoryColor = [this](const FTextBlockStyle& OriginalStyle, const FName MessageCategory)
+	{
+		FTextBlockStyle Result = OriginalStyle;
+
+		FLinearColor HSV = OriginalStyle.ColorAndOpacity.GetSpecifiedColor().LinearRGBToHSV();
+		HSV.R = GetCategoryHue(MessageCategory);
+		HSV.G = FMath::Max(0.4f, HSV.G);
+		Result.ColorAndOpacity = HSV.HSVToLinearRGB();
+		return Result;
+	};
+
 	for (int32 MessageIndex = NextPendingMessageIndex; MessageIndex < CurrentMessagesCount; ++MessageIndex)
 	{
 		const TSharedPtr<FOutputLogMessage> Message = Messages[MessageIndex];
+		const int32 LineIndex = TextLayout->GetLineModels().Num() + NumAddedMessages;
 
 		Filter->AddAvailableLogCategory(Message->Category);
 		if (!Filter->IsMessageAllowed(Message))
@@ -762,7 +878,67 @@ void FOutputLogTextLayoutMarshaller::AppendPendingMessagesToTextLayout()
 		TSharedRef<FString> LineText = Message->Message;
 
 		TArray<TSharedRef<IRun>> Runs;
-		Runs.Add(FSlateTextRun::Create(FRunInfo(), LineText, MessageTextStyle));
+
+		switch (CategoryColorizationMode)
+		{
+		case ELogCategoryColorizationMode::None:
+			Runs.Add(FSlateTextRun::Create(FRunInfo(), LineText, MessageTextStyle));
+			break;
+		case ELogCategoryColorizationMode::ColorizeWholeLine:
+			{
+				const bool bUseCategoryColor = (Message->Verbosity > ELogVerbosity::Warning);
+				Runs.Add(FSlateTextRun::Create(FRunInfo(), LineText, bUseCategoryColor ? ComputeCategoryColor(MessageTextStyle, Message->Category) : MessageTextStyle));
+			}
+			break;
+		case ELogCategoryColorizationMode::ColorizeCategoryOnly:
+			{
+				if (Message->CategoryStartIndex >= 0)
+				{
+					const int32 CategoryStartIndex = Message->CategoryStartIndex;
+					const int32 CategoryStopIndex = CategoryStartIndex + (int32)Message->Category.GetStringLength() + 1;
+					if (CategoryStartIndex > 0)
+					{
+						Runs.Add(FSlateTextRun::Create(FRunInfo(), LineText, MessageTextStyle, FTextRange(0, CategoryStartIndex)));
+					}
+					Runs.Add(FSlateTextRun::Create(FRunInfo(), LineText, ComputeCategoryColor(MessageTextStyle, Message->Category), FTextRange(CategoryStartIndex, CategoryStopIndex)));
+					Runs.Add(FSlateTextRun::Create(FRunInfo(), LineText, MessageTextStyle, FTextRange(CategoryStopIndex, LineText->Len())));
+				}
+				else
+				{
+					Runs.Add(FSlateTextRun::Create(FRunInfo(), LineText, MessageTextStyle));
+				}
+			}
+			break;
+		case ELogCategoryColorizationMode::ColorizeCategoryAsBadge:
+			{
+				if (Message->CategoryStartIndex >= 0)
+				{
+					const int32 CategoryStartIndex = Message->CategoryStartIndex;
+					const int32 CategoryStopIndex = CategoryStartIndex + (int32)Message->Category.GetStringLength();
+
+					FTextBlockStyle BadgeStyle = ComputeCategoryColor(MessageTextStyle, Message->Category);
+					Highlights.Emplace(LineIndex, FTextRange(CategoryStartIndex, CategoryStopIndex), /*Zorder=*/ -20, FCategoryBadgeHighlighter::Create(BadgeStyle.ColorAndOpacity.GetSpecifiedColor()));
+					BadgeStyle.ColorAndOpacity = FLinearColor::Black;
+
+					if (CategoryStartIndex > 0)
+					{
+						Runs.Add(FSlateTextRun::Create(FRunInfo(), LineText, MessageTextStyle, FTextRange(0, CategoryStartIndex)));
+					}
+					Runs.Add(FSlateTextRun::Create(FRunInfo(), LineText, BadgeStyle, FTextRange(CategoryStartIndex, CategoryStopIndex)));
+					Runs.Add(FSlateTextRun::Create(FRunInfo(), LineText, MessageTextStyle, FTextRange(CategoryStopIndex, LineText->Len())));
+				}
+				else
+				{
+					Runs.Add(FSlateTextRun::Create(FRunInfo(), LineText, MessageTextStyle));
+				}
+			}
+			break;
+		}
+
+		if (Message->Category.IsValid() && (Message->Category == CategoryToHighlight))
+		{
+			Highlights.Emplace(LineIndex, FTextRange(0, LineText->Len()), /*Zorder=*/ -5, FCategoryLineHighlighter::Create());
+		}
 
 		LinesToAdd.Emplace(MoveTemp(LineText), MoveTemp(Runs));
 	}
@@ -774,6 +950,11 @@ void FOutputLogTextLayoutMarshaller::AppendPendingMessagesToTextLayout()
 	}
 
 	TextLayout->AddLines(LinesToAdd);
+
+	for (const FTextLineHighlight& Highlight : Highlights)
+	{
+		TextLayout->AddLineHighlight(Highlight);
+	}
 }
 
 void FOutputLogTextLayoutMarshaller::ClearMessages()
@@ -834,6 +1015,21 @@ void FOutputLogTextLayoutMarshaller::MarkMessagesCacheAsDirty()
 	bNumMessagesCacheDirty = true;
 }
 
+FName FOutputLogTextLayoutMarshaller::GetCategoryForLocation(const FTextLocation Location) const
+{
+	if (Messages.IsValidIndex(Location.GetLineIndex()))
+	{
+		return Messages[Location.GetLineIndex()]->Category;
+	}
+
+	return NAME_None;
+}
+
+FTextLocation FOutputLogTextLayoutMarshaller::GetTextLocationAt(const FVector2D& Relative) const
+{
+	return TextLayout ? TextLayout->GetTextLocationAt(Relative) : FTextLocation(INDEX_NONE, INDEX_NONE);
+}
+
 FOutputLogTextLayoutMarshaller::FOutputLogTextLayoutMarshaller(TArray< TSharedPtr<FOutputLogMessage> > InMessages, FOutputLogFilter* InFilter)
 	: Messages(MoveTemp(InMessages))
 	, NextPendingMessageIndex(0)
@@ -976,6 +1172,10 @@ void SOutputLog::Construct( const FArguments& InArgs, bool bIsDrawerOutputLog)
 	// Remove itself on crash (crashmalloc has limited memory and echoing logs here at that point is useless).
 	FCoreDelegates::OnHandleSystemError.AddRaw(this, &SOutputLog::OnCrash);
 
+	// Listen for style changes
+	UEditorStyleSettings* Settings = GetMutableDefault<UEditorStyleSettings>();
+	SettingsWatchHandle = Settings->OnSettingChanged().AddRaw(this, &SOutputLog::HandleSettingChanged);
+	
 	bIsUserScrolled = false;
 	RequestForceScroll();
 }
@@ -988,6 +1188,12 @@ SOutputLog::~SOutputLog()
 		GLog->RemoveOutputDevice(this);
 	}
 	FCoreDelegates::OnHandleSystemError.RemoveAll(this);
+
+	if (UObjectInitialized() && !GExitPurge)
+	{
+		UEditorStyleSettings* Settings = GetMutableDefault<UEditorStyleSettings>();
+		Settings->OnSettingChanged().Remove(SettingsWatchHandle);
+	}
 }
 
 void SOutputLog::Tick(const FGeometry& AllottedGeometry, const double InCurrentTime, const float InDeltaTime)
@@ -1012,6 +1218,11 @@ void SOutputLog::OnCrash()
 	}
 }
 
+static const FName NAME_StyleLogCommand(TEXT("Log.Command"));
+static const FName NAME_StyleLogError(TEXT("Log.Error"));
+static const FName NAME_StyleLogWarning(TEXT("Log.Warning"));
+static const FName NAME_StyleLogNormal(TEXT("Log.Normal"));
+
 bool SOutputLog::CreateLogMessages( const TCHAR* V, ELogVerbosity::Type Verbosity, const class FName& Category, TArray< TSharedPtr<FOutputLogMessage> >& OutMessages )
 {
 	if (Verbosity == ELogVerbosity::SetColor)
@@ -1021,24 +1232,24 @@ bool SOutputLog::CreateLogMessages( const TCHAR* V, ELogVerbosity::Type Verbosit
 	}
 	else
 	{
-		// Get the style for this message. When piping output from child processes (eg. when cooking through the editor), we want to highlight messages
+		// Get the style for this message. When piping output from child processes (e.g., when cooking through the editor), we want to highlight messages
 		// according to their original verbosity, so also check for "Error:" and "Warning:" substrings. This is consistent with how the build system processes logs.
 		FName Style;
 		if (Category == NAME_Cmd)
 		{
-			Style = FName(TEXT("Log.Command"));
+			Style = NAME_StyleLogCommand;
 		}
 		else if (Verbosity == ELogVerbosity::Error || FCString::Stristr(V, TEXT("Error:")) != nullptr)
 		{
-			Style = FName(TEXT("Log.Error"));
+			Style = NAME_StyleLogError;
 		}
 		else if (Verbosity == ELogVerbosity::Warning || FCString::Stristr(V, TEXT("Warning:")) != nullptr)
 		{
-			Style = FName(TEXT("Log.Warning"));
+			Style = NAME_StyleLogWarning;
 		}
 		else
 		{
-			Style = FName(TEXT("Log.Normal"));
+			Style = NAME_StyleLogNormal;
 		}
 
 		// Determine how to format timestamps
@@ -1071,19 +1282,20 @@ bool SOutputLog::CreateLogMessages( const TCHAR* V, ELogVerbosity::Type Verbosit
 					int32 HardWrapLineLen = 0;
 					if (bIsFirstLineInMessage)
 					{
-						FString MessagePrefix = FOutputDeviceHelper::FormatLogLine(Verbosity, Category, nullptr, LogTimestampMode);
+						int32 CategoryStartIndex;
+						const FString MessagePrefix = FOutputDeviceHelper::FormatLogLine(Verbosity, Category, nullptr, LogTimestampMode, -1.0, /*out*/ &CategoryStartIndex);
 						
 						HardWrapLineLen = FMath::Min(HardWrapLen - MessagePrefix.Len(), Line.Len() - CurrentStartIndex);
-						FString HardWrapLine = Line.Mid(CurrentStartIndex, HardWrapLineLen);
+						const FString HardWrapLine = Line.Mid(CurrentStartIndex, HardWrapLineLen);
 
-						OutMessages.Add(MakeShared<FOutputLogMessage>(MakeShared<FString>(MessagePrefix + HardWrapLine), Verbosity, Category, Style));
+						OutMessages.Add(MakeShared<FOutputLogMessage>(MakeShared<FString>(MessagePrefix + HardWrapLine), Verbosity, Category, Style, CategoryStartIndex));
 					}
 					else
 					{
 						HardWrapLineLen = FMath::Min(HardWrapLen, Line.Len() - CurrentStartIndex);
 						FString HardWrapLine = Line.Mid(CurrentStartIndex, HardWrapLineLen);
-
-						OutMessages.Add(MakeShared<FOutputLogMessage>(MakeShared<FString>(MoveTemp(HardWrapLine)), Verbosity, Category, Style));
+						
+						OutMessages.Add(MakeShared<FOutputLogMessage>(MakeShared<FString>(MoveTemp(HardWrapLine)), Verbosity, Category, Style, INDEX_NONE));
 					}
 
 					bIsFirstLineInMessage = false;
@@ -1114,6 +1326,49 @@ void SOutputLog::ExtendTextBoxMenu(FMenuBuilder& Builder)
 		FSlateIcon(), 
 		ClearOutputLogAction
 		);
+
+	const FVector2D CursorPos = FSlateApplication::Get().GetCursorPos();
+	const FVector2D RelativeCursorPos = MessagesTextBox->GetTickSpaceGeometry().AbsoluteToLocal(CursorPos);
+	const FTextLocation CursorTextLocation = MessagesTextMarshaller->GetTextLocationAt(RelativeCursorPos);
+
+	if (CursorTextLocation.IsValid())
+	{
+		const FName CategoryName = MessagesTextMarshaller->GetCategoryForLocation(CursorTextLocation);
+
+		if (CategoryName.IsValid())
+		{
+			Builder.BeginSection(NAME_None, FText::Format(LOCTEXT("CategoryActionsSectionHeading", "Category {0}"), FText::FromName(CategoryName)));
+
+			if (CategoryName == MessagesTextMarshaller->GetCategoryToHighlight())
+			{
+				FUIAction StopHighlightingCategoryAction(
+					FExecuteAction::CreateRaw(this, &SOutputLog::OnHighlightCategory, FName())
+				);
+
+				Builder.AddMenuEntry(
+					LOCTEXT("StopHighlightCategoryAction", "Remove category highlights"),
+					LOCTEXT("StopHighlightCategoryActionTooltip", "Stop highlighting all messages for this category"),
+					FSlateIcon(),
+					StopHighlightingCategoryAction
+				);
+			}
+			else
+			{
+				FUIAction HighlightCategoryAction(
+					FExecuteAction::CreateRaw(this, &SOutputLog::OnHighlightCategory, CategoryName)
+				);
+
+				Builder.AddMenuEntry(
+					FText::Format(LOCTEXT("HighlightCategoryAction", "Highlight category {0}"), FText::FromName(CategoryName)),
+					LOCTEXT("HighlightCategoryActionTooltip", "Highlights all messages for this category"),
+					FSlateIcon(),
+					HighlightCategoryAction
+				);
+			}
+
+			Builder.EndSection();
+		}
+	}
 }
 
 void SOutputLog::OnClearLog()
@@ -1126,6 +1381,33 @@ void SOutputLog::OnClearLog()
 	bIsUserScrolled = false;
 }
 
+void SOutputLog::OnHighlightCategory(FName NewCategoryToHighlight)
+{
+	MessagesTextMarshaller->SetCategoryToHighlight(NewCategoryToHighlight);
+
+	RefreshAllPreservingLocation();
+}
+
+void SOutputLog::HandleSettingChanged(FName ChangedSettingName)
+{
+	RefreshAllPreservingLocation();
+}
+
+void SOutputLog::RefreshAllPreservingLocation()
+{
+	const FTextLocation LastCursorTextLocation = MessagesTextBox->GetCursorLocation();
+
+	MessagesTextMarshaller->MarkMessagesCacheAsDirty();
+	MessagesTextMarshaller->MakeDirty();
+	MessagesTextBox->Refresh();
+
+	//@TODO: Without this, the window will scroll if the last 'normally clicked location' is not on screen
+	// (even with the right-click set cursor pos fix, the refresh will scroll you back to the top of the screen
+	// until you left click, or to where you last left clicked otherwise if off screen; spooky...)
+	// Ideally we could read the current location or fix the bug where a refresh causes a scroll
+	MessagesTextBox->GoTo(LastCursorTextLocation);
+}
+
 void SOutputLog::OnUserScrolled(float ScrollOffset)
 {
 	bIsUserScrolled = ScrollOffset < 1.0 && !FMath::IsNearlyEqual(ScrollOffset, 1.0f);
diff --git a/Engine/Source/Developer/OutputLog/Private/SOutputLog.h b/Engine/Source/Developer/OutputLog/Private/SOutputLog.h
--- a/Engine/Source/Developer/OutputLog/Private/SOutputLog.h
+++ b/Engine/Source/Developer/OutputLog/Private/SOutputLog.h
@@ -28,20 +28,14 @@ struct FOutputLogMessage
 {
 	TSharedRef<FString> Message;
 	ELogVerbosity::Type Verbosity;
+	int8 CategoryStartIndex;
 	FName Category;
 	FName Style;
 
-	FOutputLogMessage(const TSharedRef<FString>& NewMessage, FName NewCategory, FName NewStyle = NAME_None)
-		: Message(NewMessage)
-		, Verbosity(ELogVerbosity::Log)
-		, Category(NewCategory)
-		, Style(NewStyle)
-	{
-	}
-
-	FOutputLogMessage(const TSharedRef<FString>& NewMessage, ELogVerbosity::Type NewVerbosity, FName NewCategory, FName NewStyle = NAME_None)
+	FOutputLogMessage(const TSharedRef<FString>& NewMessage, ELogVerbosity::Type NewVerbosity, FName NewCategory, FName NewStyle, int32 InCategoryStartIndex)
 		: Message(NewMessage)
 		, Verbosity(NewVerbosity)
+		, CategoryStartIndex((int8)InCategoryStartIndex)
 		, Category(NewCategory)
 		, Style(NewStyle)
 	{
@@ -351,6 +345,14 @@ public:
 	*/
 	void OnClearLog();
 
+	/** Called when a category is selected to be highlighted */
+	void OnHighlightCategory(FName NewCategoryToHighlight);
+
+	/** Called when the editor style settings are modified */
+	void HandleSettingChanged(FName ChangedSettingName);
+
+	void RefreshAllPreservingLocation();
+
 	/**
 	 * Called to determine whether delete all is currently a valid command
 	 */
@@ -465,6 +467,8 @@ protected:
 	/** Visible messages filter */
 	FOutputLogFilter Filter;
 
+	FDelegateHandle SettingsWatchHandle;
+
 	bool bIsInDrawer = false;
 };
 
@@ -492,12 +496,24 @@ public:
 
 	void MarkMessagesCacheAsDirty();
 
+	FName GetCategoryForLocation(const FTextLocation Location) const;
+
+	FTextLocation GetTextLocationAt(const FVector2D& Relative) const;
+
+	FName GetCategoryToHighlight() const { return CategoryToHighlight; }
+
+	void SetCategoryToHighlight(FName InCategory) { CategoryToHighlight = InCategory; }
+
 protected:
 
 	FOutputLogTextLayoutMarshaller(TArray< TSharedPtr<FOutputLogMessage> > InMessages, FOutputLogFilter* InFilter);
 
 	void AppendPendingMessagesToTextLayout();
 
+	TMap<FName, float> CategoryHueMap;
+
+	float GetCategoryHue(FName CategoryName);
+
 	/** All log messages to show in the text box */
 	TArray< TSharedPtr<FOutputLogMessage> > Messages;
 
@@ -513,5 +529,7 @@ protected:
 	/** Visible messages filter */
 	FOutputLogFilter* Filter;
 
+	FName CategoryToHighlight;
+
 	FTextLayout* TextLayout;
 };
diff --git a/Engine/Source/Editor/EditorStyle/Public/Classes/EditorStyleSettings.h b/Engine/Source/Editor/EditorStyle/Public/Classes/EditorStyleSettings.h
--- a/Engine/Source/Editor/EditorStyle/Public/Classes/EditorStyleSettings.h
+++ b/Engine/Source/Editor/EditorStyle/Public/Classes/EditorStyleSettings.h
@@ -29,6 +29,22 @@ enum class EAssetEditorOpenLocation : uint8
 	LastDockedWindowOrContentBrowser
 };
 
+UENUM()
+enum class ELogCategoryColorizationMode : uint8
+{
+	/** Do not colorize based on log categories */
+	None,
+
+	/** Colorize the entire log line, but not warnings or errors */
+	ColorizeWholeLine,
+
+	/** Colorize only the category name (including on warnings and errors) */
+	ColorizeCategoryOnly,
+
+	/** Colorize the background of the category name (including on warnings and errors) */
+	ColorizeCategoryAsBadge
+};
+
 /**
  * Implements the Editor style settings.
  */
@@ -141,7 +157,7 @@ public:
 	UPROPERTY(EditAnywhere, config, Category="Output Log", meta=(DisplayName="Log Font Size", ConfigRestartRequired=true))
 	int32 LogFontSize;
 
-	/** The display mode for timestamps in the output log window*/
+	/** The display mode for timestamps in the output log window */
 	UPROPERTY(EditAnywhere, config, Category="Output Log", meta=(DisplayName = "Output Log Window Timestamp Mode"))
 	TEnumAsByte<ELogTimes::Type> LogTimestampMode;
 
@@ -149,6 +165,10 @@ public:
 	UPROPERTY(EditAnywhere, config, Category="Output Log")
 	bool bPromoteOutputLogWarningsDuringPIE;
 
+	/** How should categories be colorized in the output log? */
+	UPROPERTY(EditAnywhere, config, Category = "Output Log")
+	ELogCategoryColorizationMode CategoryColorizationMode;
+
 	/**
 	 * If checked pressing the console command shortcut will cycle between focusing the status bar console, opening the output log drawer, and back to the previous focus target. 
 	 * If unchecked, the console command shortcut will only focus the status bar console
diff --git a/Engine/Source/Runtime/Core/Private/Misc/OutputDeviceHelper.cpp b/Engine/Source/Runtime/Core/Private/Misc/OutputDeviceHelper.cpp
--- a/Engine/Source/Runtime/Core/Private/Misc/OutputDeviceHelper.cpp
+++ b/Engine/Source/Runtime/Core/Private/Misc/OutputDeviceHelper.cpp
@@ -7,7 +7,7 @@
 #include "Misc/App.h"
 #include "CoreGlobals.h"
 
-FString FOutputDeviceHelper::FormatLogLine( ELogVerbosity::Type Verbosity, const class FName& Category, const TCHAR* Message /*= nullptr*/, ELogTimes::Type LogTime /*= ELogTimes::None*/, const double Time /*= -1.0*/ )
+FString FOutputDeviceHelper::FormatLogLine( ELogVerbosity::Type Verbosity, const class FName& Category, const TCHAR* Message /*= nullptr*/, ELogTimes::Type LogTime /*= ELogTimes::None*/, const double Time /*= -1.0*/, int32* OutCategoryIndex)
 {
 	const bool bShowCategory = GPrintLogCategory && Category != NAME_None;
 	FString Format;
@@ -16,7 +16,7 @@ FString FOutputDeviceHelper::FormatLogLine( ELogVerbosity::Type Verbosity, const
 	{
 		case ELogTimes::SinceGStartTime:
 		{																	
-			const double RealTime = Time == -1.0f ? FPlatformTime::Seconds() - GStartTime : Time;
+			const double RealTime = (Time < 0.0) ? (FPlatformTime::Seconds() - GStartTime) : Time;
 			Format = FString::Printf( TEXT( "[%07.2f][%3llu]" ), RealTime, GFrameCounter % 1000);
 			break;
 		}
@@ -37,6 +37,11 @@ FString FOutputDeviceHelper::FormatLogLine( ELogVerbosity::Type Verbosity, const
 			break;
 	}	
 
+	if (OutCategoryIndex != nullptr)
+	{
+		*OutCategoryIndex = bShowCategory ? Format.Len() : INDEX_NONE;
+	}
+
 	if (bShowCategory)
 	{
 		Format += Category.ToString();
diff --git a/Engine/Source/Runtime/Core/Public/Misc/OutputDeviceHelper.h b/Engine/Source/Runtime/Core/Public/Misc/OutputDeviceHelper.h
--- a/Engine/Source/Runtime/Core/Public/Misc/OutputDeviceHelper.h
+++ b/Engine/Source/Runtime/Core/Public/Misc/OutputDeviceHelper.h
@@ -27,9 +27,10 @@ struct CORE_API FOutputDeviceHelper
 	 * @param Message Optional message text. If nullptr, only the date/time/category/verbosity prefix will be returned
 	 * @param LogTime Time format
 	 * @param Time Time in seconds
+	 * @param OutCategoryIndex (if non-null) The index of the category within the return string is written here, or INDEX_NONE if the category is suppressed
 	 * @returns Formatted log line
 	 */
-	static FString FormatLogLine(ELogVerbosity::Type Verbosity, const class FName& Category, const TCHAR* Message = nullptr, ELogTimes::Type LogTime = ELogTimes::None, const double Time = -1.0);
+	static FString FormatLogLine(ELogVerbosity::Type Verbosity, const class FName& Category, const TCHAR* Message = nullptr, ELogTimes::Type LogTime = ELogTimes::None, const double Time = -1.0, int32* OutCategoryIndex = nullptr);
 
 	/**
 	 * Formats, casts to ANSI char and serializes a message to archive. Optimized for small number of allocations and Serialize calls
diff --git a/Engine/Source/Runtime/Slate/Private/Widgets/Input/SMultiLineEditableTextBox.cpp b/Engine/Source/Runtime/Slate/Private/Widgets/Input/SMultiLineEditableTextBox.cpp
--- a/Engine/Source/Runtime/Slate/Private/Widgets/Input/SMultiLineEditableTextBox.cpp
+++ b/Engine/Source/Runtime/Slate/Private/Widgets/Input/SMultiLineEditableTextBox.cpp
@@ -533,6 +533,11 @@ TArray<TSharedRef<const IRun>> SMultiLineEditableTextBox::GetSelectedRuns() cons
 	return EditableText->GetSelectedRuns();
 }
 
+FTextLocation SMultiLineEditableTextBox::GetCursorLocation() const
+{
+	return EditableText->GetCursorLocation();
+}
+
 TSharedPtr<const SScrollBar> SMultiLineEditableTextBox::GetHScrollBar() const
 {
 	return EditableText->GetHScrollBar();
diff --git a/Engine/Source/Runtime/Slate/Private/Widgets/Text/SMultiLineEditableText.cpp b/Engine/Source/Runtime/Slate/Private/Widgets/Text/SMultiLineEditableText.cpp
--- a/Engine/Source/Runtime/Slate/Private/Widgets/Text/SMultiLineEditableText.cpp
+++ b/Engine/Source/Runtime/Slate/Private/Widgets/Text/SMultiLineEditableText.cpp
@@ -532,6 +532,11 @@ TArray<TSharedRef<const IRun>> SMultiLineEditableText::GetSelectedRuns() const
 	return EditableTextLayout->GetSelectedRuns();
 }
 
+FTextLocation SMultiLineEditableText::GetCursorLocation() const
+{
+	return EditableTextLayout->GetCursorLocation();
+}
+
 TSharedPtr<const SScrollBar> SMultiLineEditableText::GetHScrollBar() const
 {
 	return HScrollBar;
diff --git a/Engine/Source/Runtime/Slate/Private/Widgets/Text/SlateEditableTextLayout.cpp b/Engine/Source/Runtime/Slate/Private/Widgets/Text/SlateEditableTextLayout.cpp
--- a/Engine/Source/Runtime/Slate/Private/Widgets/Text/SlateEditableTextLayout.cpp
+++ b/Engine/Source/Runtime/Slate/Private/Widgets/Text/SlateEditableTextLayout.cpp
@@ -2841,6 +2841,11 @@ TArray<TSharedRef<const IRun>> FSlateEditableTextLayout::GetSelectedRuns() const
 	return Runs;
 }
 
+FTextLocation FSlateEditableTextLayout::GetCursorLocation() const
+{
+	return CursorInfo.GetCursorInteractionLocation();
+}
+
 FTextLocation FSlateEditableTextLayout::TranslatedLocation(const FTextLocation& Location, int8 Direction) const
 {
 	check(Direction != 0);
diff --git a/Engine/Source/Runtime/Slate/Public/Widgets/Input/SMultiLineEditableTextBox.h b/Engine/Source/Runtime/Slate/Public/Widgets/Input/SMultiLineEditableTextBox.h
--- a/Engine/Source/Runtime/Slate/Public/Widgets/Input/SMultiLineEditableTextBox.h
+++ b/Engine/Source/Runtime/Slate/Public/Widgets/Input/SMultiLineEditableTextBox.h
@@ -441,6 +441,9 @@ public:
 	/** Get the runs currently that are current selected, some of which may be only partially selected */
 	TArray<TSharedRef<const IRun>> GetSelectedRuns() const;
 
+	/** Get the interaction position of the cursor (where to insert, delete, etc, text from/to) */
+	FTextLocation GetCursorLocation() const;
+
 	/** Get the horizontal scroll bar widget */
 	TSharedPtr<const SScrollBar> GetHScrollBar() const;
 
diff --git a/Engine/Source/Runtime/Slate/Public/Widgets/Text/SMultiLineEditableText.h b/Engine/Source/Runtime/Slate/Public/Widgets/Text/SMultiLineEditableText.h
--- a/Engine/Source/Runtime/Slate/Public/Widgets/Text/SMultiLineEditableText.h
+++ b/Engine/Source/Runtime/Slate/Public/Widgets/Text/SMultiLineEditableText.h
@@ -376,6 +376,9 @@ public:
 	/** Get the runs currently that are current selected, some of which may be only partially selected */
 	TArray<TSharedRef<const IRun>> GetSelectedRuns() const;
 
+	/** Get the interaction position of the cursor (where to insert, delete, etc, text from/to) */
+	FTextLocation GetCursorLocation() const;
+
 	/** Get the horizontal scroll bar widget */
 	TSharedPtr<const SScrollBar> GetHScrollBar() const;
 
diff --git a/Engine/Source/Runtime/Slate/Public/Widgets/Text/SlateEditableTextLayout.h b/Engine/Source/Runtime/Slate/Public/Widgets/Text/SlateEditableTextLayout.h
--- a/Engine/Source/Runtime/Slate/Public/Widgets/Text/SlateEditableTextLayout.h
+++ b/Engine/Source/Runtime/Slate/Public/Widgets/Text/SlateEditableTextLayout.h
@@ -279,6 +279,9 @@ public:
 	/** Get the runs currently that are current selected, some of which may be only partially selected */
 	TArray<TSharedRef<const IRun>> GetSelectedRuns() const;
 
+	/** Get the interaction position of the cursor (where to insert, delete, etc, text from/to) */
+	FTextLocation GetCursorLocation() const;
+
 	/**
 	 * Given a location and a Direction to offset, return a new location.
 	 *
