Add more flexible MakeOutputLogWidget function to OutputLog module

diff --git a/Engine/Source/Developer/OutputLog/Private/OutputLogModule.cpp b/Engine/Source/Developer/OutputLog/Private/OutputLogModule.cpp
--- a/Engine/Source/Developer/OutputLog/Private/OutputLogModule.cpp
+++ b/Engine/Source/Developer/OutputLog/Private/OutputLogModule.cpp
@@ -1,12 +1,12 @@
 // Copyright Epic Games, Inc. All Rights Reserved.
 
 #include "OutputLogModule.h"
-#include "Features/IModularFeatures.h"
 #include "Modules/ModuleManager.h"
 #include "Framework/Application/SlateApplication.h"
 #include "Textures/SlateIcon.h"
 #include "Framework/Docking/TabManager.h"
 #include "EditorStyleSet.h"
+#include "OutputLogCreationParams.h"
 #include "SDebugConsole.h"
 #include "SOutputLog.h"
 #include "SDeviceOutputLog.h"
@@ -166,6 +166,16 @@ TSharedRef<SWidget> FOutputLogModule::MakeOutputLogDrawerWidget(const FSimpleDel
 	return OutputLogDrawerPinned.ToSharedRef();
 }
 
+TSharedRef<SWidget> FOutputLogModule::MakeOutputLogWidget(const FOutputLogCreationParams& Params)
+{
+	return SNew(SOutputLog, Params.bCreateDockInLayoutButton)
+			.OnCloseConsole(Params.OnCloseConsole)
+			.Messages(OutputLogHistory->GetMessages())
+			.SettingsMenuFlags(Params.SettingsMenuCreationFlags)
+			.DefaultCategorySelection(Params.DefaultCategorySelection)
+			.AllowInitialLogCategory(Params.AllowAsInitialLogCategory);
+}
+
 void FOutputLogModule::ToggleDebugConsoleForWindow(const TSharedRef<SWindow>& Window, const EDebugConsoleStyle::Type InStyle, const FDebugConsoleDelegates& DebugConsoleDelegates)
 {
 	bool bShouldOpen = true;
diff --git a/Engine/Source/Developer/OutputLog/Private/SOutputLog.cpp b/Engine/Source/Developer/OutputLog/Private/SOutputLog.cpp
--- a/Engine/Source/Developer/OutputLog/Private/SOutputLog.cpp
+++ b/Engine/Source/Developer/OutputLog/Private/SOutputLog.cpp
@@ -1040,15 +1040,10 @@ FOutputLogTextLayoutMarshaller::FOutputLogTextLayoutMarshaller(TArray< TSharedPt
 }
 
 BEGIN_SLATE_FUNCTION_BUILD_OPTIMIZATION
-void SOutputLog::Construct( const FArguments& InArgs, bool bIsDrawerOutputLog)
+void SOutputLog::Construct( const FArguments& InArgs, bool bCreateDrawerDockButton)
 {
-	bIsInDrawer = bIsDrawerOutputLog;
-
-	// Build list of available log categories from historical logs
-	for (const auto& Message : InArgs._Messages)
-	{
-		Filter.AddAvailableLogCategory(Message->Category);
-	}
+	bShouldCreateDrawerDockButton = bCreateDrawerDockButton;
+	BuildInitialLogCategoryFilter(InArgs);
 
 	MessagesTextMarshaller = FOutputLogTextLayoutMarshaller::Create(InArgs._Messages, &Filter);
 
@@ -1124,7 +1119,7 @@ void SOutputLog::Construct( const FArguments& InArgs, bool bIsDrawerOutputLog)
 			[
 				SNew(SComboButton)
 				.ComboButtonStyle(FAppStyle::Get(), "SimpleComboButton")
-				.OnGetMenuContent(this, &SOutputLog::GetViewButtonContent)
+				.OnGetMenuContent(this, &SOutputLog::GetViewButtonContent, InArgs._SettingsMenuFlags)
 				.ButtonContent()
 				[
 					SNew(SHorizontalBox)
@@ -1481,6 +1476,31 @@ void SOutputLog::SetClearOnPIE(ECheckBoxState InValue)
 	GConfig->SetBool(TEXT("/Script/UnrealEd.EditorPerProjectUserSettings"), TEXT("bEnableOutputLogClearOnPIE"), ClearOnPIEEnabled, GEditorPerProjectIni);
 }
 
+void SOutputLog::BuildInitialLogCategoryFilter(const FArguments& InArgs)
+{
+	for (const auto& Message : InArgs._Messages)
+	{
+		Filter.AddAvailableLogCategory(Message->Category);
+		const bool bIsDeselectedByDefault = InArgs._AllowInitialLogCategory.IsBound() && !InArgs._AllowInitialLogCategory.Execute(Message->Category);
+		if (bIsDeselectedByDefault)
+		{
+			Filter.bShowAllCategories = false;
+			Filter.ToggleLogCategory(Message->Category);
+		}
+	}
+
+	for (auto DefaultCategorySelectionIt = InArgs._DefaultCategorySelection.CreateConstIterator(); DefaultCategorySelectionIt; ++DefaultCategorySelectionIt)
+	{
+		Filter.AddAvailableLogCategory(DefaultCategorySelectionIt->Key);
+
+		if (!DefaultCategorySelectionIt->Value)
+		{
+			Filter.bShowAllCategories = false;
+			Filter.ToggleLogCategory(DefaultCategorySelectionIt->Key);
+		}
+	}
+}
+
 void SOutputLog::OnFilterTextChanged(const FText& InFilterText)
 {
 	if (Filter.GetFilterText().ToString().Equals(InFilterText.ToString(), ESearchCase::CaseSensitive))
@@ -1679,7 +1699,7 @@ void SOutputLog::CategoriesSingle_Execute(FName InName)
 	Refresh();
 }
 
-TSharedRef<SWidget> SOutputLog::GetViewButtonContent()
+TSharedRef<SWidget> SOutputLog::GetViewButtonContent(EOutputLogSettingsMenuFlags Flags)
 {
 	TSharedPtr<FExtender> Extender;
 	FMenuBuilder MenuBuilder(true, nullptr, Extender, true);
@@ -1698,21 +1718,25 @@ TSharedRef<SWidget> SOutputLog::GetViewButtonContent()
 		NAME_None,
 		EUserInterfaceActionType::ToggleButton
 	);
-	MenuBuilder.AddMenuEntry(
-		LOCTEXT("ClearOnPIE", "Clear on PIE"),
-		LOCTEXT("ClearOnPIEToolTip", "Enable clearing of the Output Log on PIE startup."),
-		FSlateIcon(),
-		FUIAction(
-			FExecuteAction::CreateLambda([this] {
-				// This is a toggle, hence that it is inverted
-				SetClearOnPIE(IsClearOnPIEEnabled() ? ECheckBoxState::Unchecked : ECheckBoxState::Checked);
-			}),
-			FCanExecuteAction::CreateLambda([] { return true; }),
-			FIsActionChecked::CreateSP(this, &SOutputLog::IsClearOnPIEEnabled)
-		),
-		NAME_None,
-		EUserInterfaceActionType::ToggleButton
-	);
+
+	if ((Flags & EOutputLogSettingsMenuFlags::SkipClearOnPie) == EOutputLogSettingsMenuFlags::None)
+	{
+		MenuBuilder.AddMenuEntry(
+			LOCTEXT("ClearOnPIE", "Clear on PIE"),
+			LOCTEXT("ClearOnPIEToolTip", "Enable clearing of the Output Log on PIE startup."),
+			FSlateIcon(),
+			FUIAction(
+				FExecuteAction::CreateLambda([this] {
+					// This is a toggle, hence that it is inverted
+					SetClearOnPIE(IsClearOnPIEEnabled() ? ECheckBoxState::Unchecked : ECheckBoxState::Checked);
+				}),
+				FCanExecuteAction::CreateLambda([] { return true; }),
+				FIsActionChecked::CreateSP(this, &SOutputLog::IsClearOnPIEEnabled)
+			),
+			NAME_None,
+			EUserInterfaceActionType::ToggleButton
+		);
+	}
 	MenuBuilder.AddMenuSeparator();
 
 	//Show Source In Explorer
@@ -1741,7 +1765,7 @@ TSharedRef<SWidget> SOutputLog::GetViewButtonContent()
 
 TSharedRef<SWidget> SOutputLog::CreateDrawerDockButton()
 {
-	if (bIsInDrawer)
+	if (bShouldCreateDrawerDockButton)
 	{
 		return
 			SNew(SButton)
@@ -1862,7 +1886,7 @@ bool FOutputLogFilter::IsMessageAllowed(const TSharedPtr<FOutputLogMessage>& Mes
 	return true;
 }
 
-void FOutputLogFilter::AddAvailableLogCategory(FName& LogCategory)
+void FOutputLogFilter::AddAvailableLogCategory(const FName& LogCategory)
 {
 	// Use an insert-sort to keep AvailableLogCategories alphabetically sorted
 	int32 InsertIndex = 0;
diff --git a/Engine/Source/Developer/OutputLog/Private/SOutputLog.h b/Engine/Source/Developer/OutputLog/Private/SOutputLog.h
--- a/Engine/Source/Developer/OutputLog/Private/SOutputLog.h
+++ b/Engine/Source/Developer/OutputLog/Private/SOutputLog.h
@@ -6,6 +6,7 @@
 #include "SlateFwd.h"
 #include "Widgets/DeclarativeSyntaxSupport.h"
 #include "Input/Reply.h"
+#include "OutputLogCreationParams.h"
 #include "Widgets/SWidget.h"
 #include "Widgets/SCompoundWidget.h"
 #include "Widgets/Input/SMultiLineEditableTextBox.h"
@@ -271,7 +272,7 @@ struct FOutputLogFilter
 	const TArray<FName>& GetAvailableLogCategories() { return AvailableLogCategories; }
 
 	/** Adds a Log Category to the list of available categories, if it isn't already present */
-	void AddAvailableLogCategory(FName& LogCategory);
+	void AddAvailableLogCategory(const FName& LogCategory);
 
 	/** Enables or disables a Log Category in the filter */
 	void ToggleLogCategory(const FName& LogCategory);
@@ -312,6 +313,14 @@ public:
 		/** All messages captured before this log window has been created */
 		SLATE_ARGUMENT( TArray< TSharedPtr<FOutputLogMessage> >, Messages )
 
+		/**  */
+		SLATE_ARGUMENT( EOutputLogSettingsMenuFlags, SettingsMenuFlags)
+
+		SLATE_ARGUMENT( FDefaultCategorySelectionMap, DefaultCategorySelection )
+
+		/** Used to determine the set of initially discovered log categories that should be selected */
+		SLATE_EVENT( FAllowLogCategoryCallback, AllowInitialLogCategory )
+
 	SLATE_END_ARGS()
 
 	/** Destructor for output log, so we can unregister from notifications */
@@ -322,7 +331,7 @@ public:
 	 *
 	 * @param	InArgs	Declaration used by the SNew() macro to construct this widget
 	 */
-	void Construct( const FArguments& InArgs, bool bIsDrawerOutputLog );
+	void Construct( const FArguments& InArgs, bool bCreateDrawerDockButton );
 
 	// SWidget interface
 	virtual void Tick(const FGeometry& AllottedGeometry, const double InCurrentTime, const float InDeltaTime) override;
@@ -397,6 +406,9 @@ protected:
 	bool bIsUserScrolled;
 
 private:
+
+	void BuildInitialLogCategoryFilter(const FArguments& InArgs);
+	
 	/** Called by Slate when the filter box changes text. */
 	void OnFilterTextChanged(const FText& InFilterText);
 
@@ -452,7 +464,7 @@ private:
 
 	FSlateColor GetViewButtonForegroundColor() const;
 
-	TSharedRef<SWidget> GetViewButtonContent();
+	TSharedRef<SWidget> GetViewButtonContent(EOutputLogSettingsMenuFlags Flags);
 
 	TSharedRef<SWidget> CreateDrawerDockButton();
 
@@ -469,7 +481,7 @@ protected:
 
 	FDelegateHandle SettingsWatchHandle;
 
-	bool bIsInDrawer = false;
+	bool bShouldCreateDrawerDockButton = false;
 };
 
 /** Output log text marshaller to convert an array of FOutputLogMessages into styled lines to be consumed by an FTextLayout */
diff --git a/Engine/Source/Developer/OutputLog/Public/OutputLogCreationParams.h b/Engine/Source/Developer/OutputLog/Public/OutputLogCreationParams.h
new file mode 100644
--- /dev/null
+++ b/Engine/Source/Developer/OutputLog/Public/OutputLogCreationParams.h
@@ -0,0 +1,34 @@
+// Copyright Epic Games, Inc. All Rights Reserved.
+
+#pragma once
+
+#include "CoreMinimal.h"
+
+enum class EOutputLogSettingsMenuFlags
+{
+	None = 0x00,
+
+	/** The clear on Pie button should not be created */
+	SkipClearOnPie = 0x01
+};
+ENUM_CLASS_FLAGS(EOutputLogSettingsMenuFlags)
+
+DECLARE_DELEGATE_RetVal_OneParam(bool, FAllowLogCategoryCallback, const FName);
+using FDefaultCategorySelectionMap = TMap<FName, bool>;
+
+struct FOutputLogCreationParams
+{
+	/** Whether to create the button for docking the log */
+	bool bCreateDockInLayoutButton = false;
+
+	/** Determines what entries the Settings drop-down will ignore */
+	EOutputLogSettingsMenuFlags SettingsMenuCreationFlags = EOutputLogSettingsMenuFlags::None;
+
+	/** Called when building the initial set of selected log categories */
+	FAllowLogCategoryCallback AllowAsInitialLogCategory;
+
+	/** Maps each log category to whether it should be selected or deselected by default. The caller is responsible to enter valid category names. */
+	FDefaultCategorySelectionMap DefaultCategorySelection;
+	
+	FSimpleDelegate OnCloseConsole;
+};
\ No newline at end of file
diff --git a/Engine/Source/Developer/OutputLog/Public/OutputLogModule.h b/Engine/Source/Developer/OutputLog/Public/OutputLogModule.h
--- a/Engine/Source/Developer/OutputLog/Public/OutputLogModule.h
+++ b/Engine/Source/Developer/OutputLog/Public/OutputLogModule.h
@@ -6,14 +6,13 @@
 #include "Widgets/SWidget.h"
 #include "Modules/ModuleInterface.h"
 #include "Widgets/SWindow.h"
-#include "Features/IModularFeatures.h"
-#include "HAL/IConsoleManager.h"
 
 class SMultiLineEditableTextBox;
 class FOutputLogHistory;
 class SOutputLog;
 class SDockTab;
 class FSpawnTabArgs;
+struct FOutputLogCreationParams;
 
 /** Style of the debug console */
 namespace EDebugConsoleStyle
@@ -46,9 +45,11 @@ public:
 	/** Generates a console input box widget.  Remember, this widget will become invalid if the
 		output log DLL is unloaded on the fly. */
 	virtual TSharedRef<SWidget> MakeConsoleInputBox(TSharedPtr<SMultiLineEditableTextBox>& OutExposedEditableTextBox, const FSimpleDelegate& OnCloseConsole, const FSimpleDelegate& OnConsoleCommandExecuted) const;
-
+	
 	virtual TSharedRef<SWidget> MakeOutputLogDrawerWidget(const FSimpleDelegate& OnCloseConsole);
 
+	virtual TSharedRef<SWidget> MakeOutputLogWidget(const FOutputLogCreationParams& Params);
+	
 	/** Opens a debug console in the specified window, if not already open */
 	virtual void ToggleDebugConsoleForWindow(const TSharedRef<SWindow>& Window, const EDebugConsoleStyle::Type InStyle, const FDebugConsoleDelegates& DebugConsoleDelegates);
 
