Allow external code to change Output Log's filter

diff --git a/Engine/Source/Developer/OutputLog/Private/OutputLogModule.cpp b/Engine/Source/Developer/OutputLog/Private/OutputLogModule.cpp
--- a/Engine/Source/Developer/OutputLog/Private/OutputLogModule.cpp
+++ b/Engine/Source/Developer/OutputLog/Private/OutputLogModule.cpp
@@ -296,3 +296,16 @@ const TSharedPtr<SWidget> FOutputLogModule::GetOutputLog() const
 {
 	return OutputLog.Pin();
 }
+
+void FOutputLogModule::UpdateOutputLogFilter(const TArray<FName>& CategoriesToShow, TOptional<bool> bShowErrors, TOptional<bool> bShowWarnings, TOptional<bool> bShowLogs)
+{
+	if (TSharedPtr<SOutputLog> SharedOutputLog = OutputLog.Pin())
+	{
+		SharedOutputLog->UpdateOutputLogFilter(CategoriesToShow, bShowErrors, bShowWarnings, bShowLogs);
+	}
+}
+
+void FOutputLogModule::OpenOutputLog() const
+{
+	FGlobalTabmanager::Get()->TryInvokeTab(OutputLogModule::OutputLogTabName);
+}
diff --git a/Engine/Source/Developer/OutputLog/Private/SOutputLog.cpp b/Engine/Source/Developer/OutputLog/Private/SOutputLog.cpp
--- a/Engine/Source/Developer/OutputLog/Private/SOutputLog.cpp
+++ b/Engine/Source/Developer/OutputLog/Private/SOutputLog.cpp
@@ -1699,6 +1699,35 @@ void SOutputLog::CategoriesSingle_Execute(FName InName)
 	Refresh();
 }
 
+void SOutputLog::UpdateOutputLogFilter(const TArray<FName>& CategoriesToShow, TOptional<bool> bShowErrors, TOptional<bool> bShowWarnings, TOptional<bool> bShowLogs)
+{
+	if (bShowErrors.IsSet())
+	{
+		Filter.bShowErrors = bShowErrors.GetValue();
+	}
+	if (bShowWarnings.IsSet())
+	{
+		Filter.bShowWarnings = bShowWarnings.GetValue();
+	}
+	if (bShowLogs.IsSet())
+	{
+		Filter.bShowLogs = bShowLogs.GetValue();
+	}
+
+	// Show all categories if empty list is passed in. This means there's no way to change the Show* bools
+	// and also leave the categories as-is, but that use case is probably minimal. It can be added if needed.
+	Filter.bShowAllCategories = CategoriesToShow.Num() == 0;
+	Filter.ClearSelectedLogCategories();
+	const TArray<FName> ActualCategoriesToShow = Filter.bShowAllCategories ? Filter.GetAvailableLogCategories() : CategoriesToShow;
+	for (const auto& AvailableCategory : ActualCategoriesToShow)
+	{
+		Filter.ToggleLogCategory(AvailableCategory);
+	}
+
+	MessagesTextMarshaller->MarkMessagesCacheAsDirty();
+	Refresh();
+}
+
 TSharedRef<SWidget> SOutputLog::GetViewButtonContent(EOutputLogSettingsMenuFlags Flags)
 {
 	TSharedPtr<FExtender> Extender;
diff --git a/Engine/Source/Developer/OutputLog/Private/SOutputLog.h b/Engine/Source/Developer/OutputLog/Private/SOutputLog.h
--- a/Engine/Source/Developer/OutputLog/Private/SOutputLog.h
+++ b/Engine/Source/Developer/OutputLog/Private/SOutputLog.h
@@ -369,6 +369,9 @@ public:
 
 	/** Focuses the edit box where you type in console commands */
 	void FocusConsoleCommandBox();
+
+	/** Change the output log's filter. If CategoriesToShow is empty, all categories will be shown. */
+	void UpdateOutputLogFilter(const TArray<FName>& CategoriesToShow, TOptional<bool> bShowErrors = TOptional<bool>(), TOptional<bool> bShowWarnings = TOptional<bool>(), TOptional<bool> bShowLogs = TOptional<bool>());
 protected:
 
 	virtual void Serialize( const TCHAR* V, ELogVerbosity::Type Verbosity, const class FName& Category ) override;
diff --git a/Engine/Source/Developer/OutputLog/Public/OutputLogModule.h b/Engine/Source/Developer/OutputLog/Public/OutputLogModule.h
--- a/Engine/Source/Developer/OutputLog/Public/OutputLogModule.h
+++ b/Engine/Source/Developer/OutputLog/Public/OutputLogModule.h
@@ -63,6 +63,11 @@ public:
 	virtual const TSharedPtr<SWidget> GetOutputLog() const;
 	const TSharedPtr<SDockTab> GetOutputLogTab() const { return OutputLogTab.Pin(); }
 
+	/** Change the output log's filter. If CategoriesToShow is empty, all categories will be shown. */
+	void OUTPUTLOG_API UpdateOutputLogFilter(const TArray<FName>& CategoriesToShow, TOptional<bool> bShowErrors = TOptional<bool>(), TOptional<bool> bShowWarnings = TOptional<bool>(), TOptional<bool> bShowLogs = TOptional<bool>());
+	/** Opens the output log tab, or brings it to front if it's already open */
+	void OUTPUTLOG_API OpenOutputLog() const;
+
 private:
 	TSharedRef<SDockTab> SpawnOutputLogTab(const FSpawnTabArgs& Args);
 	TSharedRef<SDockTab> SpawnDeviceOutputLogTab(const FSpawnTabArgs& Args);
