Fix log categories filtering not working

diff --git a/Engine/Plugins/Developer/Concert/ConcertUI/ConcertServerUI/Source/ConcertServerUI/Private/Widgets/StatusBar/SConcertStatusBar.cpp b/Engine/Plugins/Developer/Concert/ConcertUI/ConcertServerUI/Source/ConcertServerUI/Private/Widgets/StatusBar/SConcertStatusBar.cpp
--- a/Engine/Plugins/Developer/Concert/ConcertUI/ConcertServerUI/Source/ConcertServerUI/Private/Widgets/StatusBar/SConcertStatusBar.cpp
+++ b/Engine/Plugins/Developer/Concert/ConcertUI/ConcertServerUI/Source/ConcertServerUI/Private/Widgets/StatusBar/SConcertStatusBar.cpp
@@ -2,6 +2,8 @@
 
 #include "SConcertStatusBar.h"
 
+#include "OutputLogCreationParams.h"
+
 #include "Framework/Application/SlateApplication.h"
 #include "OutputLogModule.h"
 #include "SWidgetDrawer.h"
@@ -23,7 +25,26 @@ namespace UE::ConcertServerUI::Private
 		{
 			if (!StatusBarOutputLog)
 			{
-				StatusBarOutputLog = FOutputLogModule::Get().MakeOutputLogDrawerWidget(FSimpleDelegate());
+				FOutputLogCreationParams Params;
+				Params.bCreateDockInLayoutButton = true;
+				Params.SettingsMenuCreationFlags = EOutputLogSettingsMenuFlags::SkipClearOnPie
+					| EOutputLogSettingsMenuFlags::SkipOpenSourceButton
+					| EOutputLogSettingsMenuFlags::SkipEnableWordWrapping; // Checkbox relies on saving an editor config file and does not work correctly
+				Params.AllowAsInitialLogCategory = FAllowLogCategoryCallback::CreateLambda([](const FName CategoryName)
+				{
+					// Hide these categories for improved UX
+					static TArray<FString> DeselectedByDefault = { "Slate", "WindowsTextInput" };
+					const FString CategoryNameAsString = CategoryName.ToString();
+					for (const FString& DeselectedCategory : DeselectedByDefault)
+					{
+						if (CategoryNameAsString.Contains(DeselectedCategory))
+						{
+							return false;
+						}
+					}
+					return true;
+				});
+				StatusBarOutputLog = FOutputLogModule::Get().MakeOutputLogWidget(Params);
 			}
 
 			return StatusBarOutputLog.ToSharedRef();
diff --git a/Engine/Source/Developer/OutputLog/Private/SOutputLog.cpp b/Engine/Source/Developer/OutputLog/Private/SOutputLog.cpp
--- a/Engine/Source/Developer/OutputLog/Private/SOutputLog.cpp
+++ b/Engine/Source/Developer/OutputLog/Private/SOutputLog.cpp
@@ -1482,7 +1482,7 @@ void SOutputLog::BuildInitialLogCategoryFilter(const FArguments& InArgs)
 	{
 		Filter.AddAvailableLogCategory(Message->Category);
 		const bool bIsDeselectedByDefault = InArgs._AllowInitialLogCategory.IsBound() && !InArgs._AllowInitialLogCategory.Execute(Message->Category);
-		if (bIsDeselectedByDefault)
+		if (bIsDeselectedByDefault && Filter.IsLogCategoryEnabled(Message->Category))
 		{
 			Filter.bShowAllCategories = false;
 			Filter.ToggleLogCategory(Message->Category);
@@ -1491,12 +1491,13 @@ void SOutputLog::BuildInitialLogCategoryFilter(const FArguments& InArgs)
 
 	for (auto DefaultCategorySelectionIt = InArgs._DefaultCategorySelection.CreateConstIterator(); DefaultCategorySelectionIt; ++DefaultCategorySelectionIt)
 	{
-		Filter.AddAvailableLogCategory(DefaultCategorySelectionIt->Key);
+		const FName Category = DefaultCategorySelectionIt->Key;
+		Filter.AddAvailableLogCategory(Category);
 
-		if (!DefaultCategorySelectionIt->Value)
+		Filter.bShowAllCategories &= DefaultCategorySelectionIt->Value;
+		if (Filter.IsLogCategoryEnabled(Category) != DefaultCategorySelectionIt->Value)
 		{
-			Filter.bShowAllCategories = false;
-			Filter.ToggleLogCategory(DefaultCategorySelectionIt->Key);
+			Filter.ToggleLogCategory(Category);
 		}
 	}
 }
@@ -1732,23 +1733,29 @@ TSharedRef<SWidget> SOutputLog::GetViewButtonContent(EOutputLogSettingsMenuFlags
 {
 	TSharedPtr<FExtender> Extender;
 	FMenuBuilder MenuBuilder(true, nullptr, Extender, true);
-	MenuBuilder.AddMenuEntry(
-		LOCTEXT("WordWrapEnabledOption", "Enable Word Wrapping"),
-		LOCTEXT("WordWrapEnabledOptionToolTip", "Enable word wrapping in the Output Log."),
-		FSlateIcon(),
-		FUIAction(
-			FExecuteAction::CreateLambda([this] {
-				// This is a toggle, hence that it is inverted
-				SetWordWrapEnabled(IsWordWrapEnabled() ? ECheckBoxState::Unchecked : ECheckBoxState::Checked);
-			}),
-			FCanExecuteAction::CreateLambda([] { return true; }),
-			FIsActionChecked::CreateSP(this, &SOutputLog::IsWordWrapEnabled)
-		),
-		NAME_None,
-		EUserInterfaceActionType::ToggleButton
-	);
 
-	if ((Flags & EOutputLogSettingsMenuFlags::SkipClearOnPie) == EOutputLogSettingsMenuFlags::None)
+	const bool bSupportWordWrapping = (Flags & EOutputLogSettingsMenuFlags::SkipEnableWordWrapping) == EOutputLogSettingsMenuFlags::None;
+	if (bSupportWordWrapping)
+	{
+		MenuBuilder.AddMenuEntry(
+			LOCTEXT("WordWrapEnabledOption", "Enable Word Wrapping"),
+			LOCTEXT("WordWrapEnabledOptionToolTip", "Enable word wrapping in the Output Log."),
+			FSlateIcon(),
+			FUIAction(
+				FExecuteAction::CreateLambda([this] {
+					// This is a toggle, hence that it is inverted
+					SetWordWrapEnabled(IsWordWrapEnabled() ? ECheckBoxState::Unchecked : ECheckBoxState::Checked);
+				}),
+				FCanExecuteAction::CreateLambda([] { return true; }),
+				FIsActionChecked::CreateSP(this, &SOutputLog::IsWordWrapEnabled)
+			),
+			NAME_None,
+			EUserInterfaceActionType::ToggleButton
+		);
+	}
+
+	const bool bSupportClearOnPie = (Flags & EOutputLogSettingsMenuFlags::SkipClearOnPie) == EOutputLogSettingsMenuFlags::None;
+	if (bSupportClearOnPie)
 	{
 		MenuBuilder.AddMenuEntry(
 			LOCTEXT("ClearOnPIE", "Clear on PIE"),
@@ -1766,28 +1773,39 @@ TSharedRef<SWidget> SOutputLog::GetViewButtonContent(EOutputLogSettingsMenuFlags
 			EUserInterfaceActionType::ToggleButton
 		);
 	}
-	MenuBuilder.AddMenuSeparator();
 
-	//Show Source In Explorer
-	MenuBuilder.AddMenuEntry(
-		LOCTEXT("FindSourceFile", "Open Source Location"),
-		LOCTEXT("FindSourceFileTooltip", "Opens the folder containing the source of the Output Log."),
-		FSlateIcon(FEditorStyle::GetStyleSetName(), "OutputLog.OpenSourceLocation"),
-		FUIAction(
-			FExecuteAction::CreateSP(this, &SOutputLog::OpenLogFileInExplorer)
-		)
-	);
-	
-	// Open In External Editor
-	MenuBuilder.AddMenuEntry(
-		LOCTEXT("OpenInExternalEditor", "Open In External Editor"),
-		LOCTEXT("OpenInExternalEditorTooltip", "Opens the Output Log in the default external editor."),
-		FSlateIcon(FEditorStyle::GetStyleSetName(), "OutputLog.OpenInExternalEditor"),
-		FUIAction(
-			FExecuteAction::CreateSP(this, &SOutputLog::OpenLogFileInExternalEditor)
-		)
-	);
+	const bool bSupportSourceLocation = (Flags & EOutputLogSettingsMenuFlags::SkipOpenSourceButton) == EOutputLogSettingsMenuFlags::None;
+	const bool bSupportExternalEditor = (Flags & EOutputLogSettingsMenuFlags::SkipOpenInExternalEditorButton) == EOutputLogSettingsMenuFlags::None;
+	const bool bNeedsSeparator = (bSupportWordWrapping || bSupportClearOnPie) && (bSupportSourceLocation || bSupportExternalEditor);
+	if (bNeedsSeparator)
+	{
+		MenuBuilder.AddMenuSeparator();
+	}
 
+	if (bSupportSourceLocation)
+	{
+		MenuBuilder.AddMenuEntry(
+			LOCTEXT("FindSourceFile", "Open Source Location"),
+			LOCTEXT("FindSourceFileTooltip", "Opens the folder containing the source of the Output Log."),
+			FSlateIcon(FEditorStyle::GetStyleSetName(), "OutputLog.OpenSourceLocation"),
+			FUIAction(
+				FExecuteAction::CreateSP(this, &SOutputLog::OpenLogFileInExplorer)
+			)
+		);
+	}
+	
+	if (bSupportExternalEditor)
+	{
+		// Open In External Editor
+		MenuBuilder.AddMenuEntry(
+			LOCTEXT("OpenInExternalEditor", "Open In External Editor"),
+			LOCTEXT("OpenInExternalEditorTooltip", "Opens the Output Log in the default external editor."),
+			FSlateIcon(FEditorStyle::GetStyleSetName(), "OutputLog.OpenInExternalEditor"),
+			FUIAction(
+				FExecuteAction::CreateSP(this, &SOutputLog::OpenLogFileInExternalEditor)
+			)
+		);
+	}
 
 	return MenuBuilder.MakeWidget();
 }
diff --git a/Engine/Source/Developer/OutputLog/Public/OutputLogCreationParams.h b/Engine/Source/Developer/OutputLog/Public/OutputLogCreationParams.h
--- a/Engine/Source/Developer/OutputLog/Public/OutputLogCreationParams.h
+++ b/Engine/Source/Developer/OutputLog/Public/OutputLogCreationParams.h
@@ -9,7 +9,15 @@ enum class EOutputLogSettingsMenuFlags
 	None = 0x00,
 
 	/** The clear on Pie button should not be created */
-	SkipClearOnPie = 0x01
+	SkipClearOnPie = 0x01,
+	/** The Enable world wrapping button should not be created */
+	SkipEnableWordWrapping = 0x02,
+
+	/** Skip the button that opens the source folder of the output log module */
+	SkipOpenSourceButton = 0x03,
+	/** Skip the button which opens the log in a text editor */
+	SkipOpenInExternalEditorButton = 0x04
+	
 };
 ENUM_CLASS_FLAGS(EOutputLogSettingsMenuFlags)
 
