Fix compile errors when WITH_EDITOR is 0
#rnx

diff --git a/Engine/Source/Developer/OutputLog/Private/SOutputLog.cpp b/Engine/Source/Developer/OutputLog/Private/SOutputLog.cpp
--- a/Engine/Source/Developer/OutputLog/Private/SOutputLog.cpp
+++ b/Engine/Source/Developer/OutputLog/Private/SOutputLog.cpp
@@ -1771,6 +1771,7 @@ TSharedRef<SWidget> SOutputLog::GetViewButtonContent(EOutputLogSettingsMenuFlags
 		);
 	}
 
+#if WITH_EDITOR
 	const bool bSupportClearOnPie = (Flags & EOutputLogSettingsMenuFlags::SkipClearOnPie) == EOutputLogSettingsMenuFlags::None;
 	if (bSupportClearOnPie)
 	{
@@ -1790,6 +1791,9 @@ TSharedRef<SWidget> SOutputLog::GetViewButtonContent(EOutputLogSettingsMenuFlags
 			EUserInterfaceActionType::ToggleButton
 		);
 	}
+#else
+	constexpr bool bSupportClearOnPie = false;
+#endif
 
 	const bool bSupportSourceLocation = (Flags & EOutputLogSettingsMenuFlags::SkipOpenSourceButton) == EOutputLogSettingsMenuFlags::None;
 	const bool bSupportExternalEditor = (Flags & EOutputLogSettingsMenuFlags::SkipOpenInExternalEditorButton) == EOutputLogSettingsMenuFlags::None;
