Editor: Ctrl + Console Key (` by default) now iterates through active Console Command Executors, so you can switch between Cmd and Python by hitting Ctrl + `.

diff --git a/Engine/Source/Developer/OutputLog/Private/SOutputLog.cpp b/Engine/Source/Developer/OutputLog/Private/SOutputLog.cpp
--- a/Engine/Source/Developer/OutputLog/Private/SOutputLog.cpp
+++ b/Engine/Source/Developer/OutputLog/Private/SOutputLog.cpp
@@ -28,6 +28,7 @@
 #include "Framework/Docking/TabManager.h"
 #include "Widgets/Docking/SDockTab.h"
 #include "OutputLogModule.h"
+#include "Toolkits/GlobalEditorCommonCommands.h"
 #include "Widgets/Text/SlateEditableTextTypes.h"
 #include "OutputLogSettings.h"
 #include "OutputLogStyle.h"
@@ -459,6 +460,9 @@
 	}
 	else
 	{
+		const TSharedRef<const FInputChord> SelectNextExecutorChord = FGlobalEditorCommonCommands::Get().SelectNextConsoleExecutor->GetActiveChord(EMultipleKeyBindingIndex::Primary);
+		EModifierKey::Type SelectNextExecutorModiferKey = EModifierKey::FromBools(SelectNextExecutorChord->bCtrl, SelectNextExecutorChord->bAlt, SelectNextExecutorChord->bShift, SelectNextExecutorChord->bCmd);
+
 		if(KeyEvent.GetKey() == EKeys::Up)
 		{
 			// If the command field isn't empty we need you to have pressed Control+Up to summon the history (to make sure you're not just using caret navigation)
@@ -508,6 +512,11 @@
 
 			return FReply::Handled();
 		}
+		else if (KeyEvent.GetKey() == SelectNextExecutorChord->Key && KeyEvent.GetModifierKeys().AreModifersDown(SelectNextExecutorModiferKey))
+		{
+			MakeNextCommandExecutorActive();
+			return FReply::Handled();
+		}
 	}
 
 	return FReply::Unhandled();
@@ -660,6 +669,28 @@
 	return !ConsoleCommandCustomExec.IsBound(); // custom execs always show the default executor in the UI (which has the selector disabled)
 }
 
+void SConsoleInputBox::MakeNextCommandExecutorActive()
+{
+	// Sorted so the iteration order matches the displayed order.
+	TArray<IConsoleCommandExecutor*> CommandExecutors = IModularFeatures::Get().GetModularFeatureImplementations<IConsoleCommandExecutor>(IConsoleCommandExecutor::ModularFeatureName());
+	CommandExecutors.Sort([](IConsoleCommandExecutor& LHS, IConsoleCommandExecutor& RHS)
+		{
+			return LHS.GetDisplayName().CompareTo(RHS.GetDisplayName()) < 0;
+		});
+
+	int32 CurrentIndex = CommandExecutors.IndexOfByKey(ActiveCommandExecutor);
+	if (CurrentIndex >= 0)
+	{
+		CurrentIndex++;
+		if (CurrentIndex >= CommandExecutors.Num())
+		{
+			CurrentIndex = 0;
+		}
+
+		SetActiveCommandExecutor(CommandExecutors[CurrentIndex]->GetName());
+	}
+}
+
 TSharedRef<SWidget> SConsoleInputBox::GetCommandExecutorMenuContent()
 {
 	FMenuBuilder MenuBuilder(/*bInShouldCloseWindowAfterMenuSelection=*/true, nullptr);
diff --git a/Engine/Source/Developer/OutputLog/Private/SOutputLog.h b/Engine/Source/Developer/OutputLog/Private/SOutputLog.h
--- a/Engine/Source/Developer/OutputLog/Private/SOutputLog.h
+++ b/Engine/Source/Developer/OutputLog/Private/SOutputLog.h
@@ -125,6 +125,8 @@
 
 	void SetActiveCommandExecutor(const FName InExecName);
 
+	void MakeNextCommandExecutorActive();
+
 	FText GetActiveCommandExecutorDisplayName() const;
 
 	FText GetActiveCommandExecutorHintText() const;
diff --git a/Engine/Source/Editor/UnrealEd/Private/Toolkits/GlobalEditorCommonCommands.cpp b/Engine/Source/Editor/UnrealEd/Private/Toolkits/GlobalEditorCommonCommands.cpp
--- a/Engine/Source/Editor/UnrealEd/Private/Toolkits/GlobalEditorCommonCommands.cpp
+++ b/Engine/Source/Editor/UnrealEd/Private/Toolkits/GlobalEditorCommonCommands.cpp
@@ -49,6 +49,7 @@ void FGlobalEditorCommonCommands::RegisterCommands()
 	UI_COMMAND(FindInContentBrowser, "Browse to Asset", "Browses to the associated asset and selects it in the most recently used Content Browser (summoning one if necessary)", EUserInterfaceActionType::Button, FInputChord(EModifierKey::Control, EKeys::B));
 
 	UI_COMMAND(OpenConsoleCommandBox, "Open Console Command Box", "Opens an edit box where you can type in a console command", EUserInterfaceActionType::Button, FInputChord(EKeys::Tilde));
+	UI_COMMAND(SelectNextConsoleExecutor, "Iterate Console Executor", "Iterates through active Console Executors (Python, Cmd, etc.)", EUserInterfaceActionType::Button, FInputChord(EModifierKey::Control, EKeys::Tilde));
 	
 	UI_COMMAND(OpenOutputLogDrawer, "Open Output Log Drawer", "Opens the output log drawer from the active asset editor status bar", EUserInterfaceActionType::Button, FInputChord(EModifierKey::Alt, EKeys::Tilde));
 
diff --git a/Engine/Source/Editor/UnrealEd/Public/Toolkits/GlobalEditorCommonCommands.h b/Engine/Source/Editor/UnrealEd/Public/Toolkits/GlobalEditorCommonCommands.h
--- a/Engine/Source/Editor/UnrealEd/Public/Toolkits/GlobalEditorCommonCommands.h
+++ b/Engine/Source/Editor/UnrealEd/Public/Toolkits/GlobalEditorCommonCommands.h
@@ -43,6 +43,7 @@ public:
 	TSharedPtr<FUICommandInfo> SummonOpenAssetDialogAlternate;
 	TSharedPtr<FUICommandInfo> OpenDocumentation;
 	TSharedPtr<FUICommandInfo> OpenConsoleCommandBox;
+	TSharedPtr<FUICommandInfo> SelectNextConsoleExecutor;
 	TSharedPtr<FUICommandInfo> OpenOutputLogDrawer;
 	TSharedPtr<FUICommandInfo> OpenContentBrowserDrawer;
 };
