Editor: Ctrl + Console Key (` by default) now iterates through active Console Command Executors, so you can switch between Cmd and Python by hitting Ctrl + `.

diff --git a/Engine/Plugins/Developer/Concert/ConcertApp/MultiUserServer/Source/MultiUserServer/Private/ConcertConsoleCommandExecutor.cpp b/Engine/Plugins/Developer/Concert/ConcertApp/MultiUserServer/Source/MultiUserServer/Private/ConcertConsoleCommandExecutor.cpp
--- a/Engine/Plugins/Developer/Concert/ConcertApp/MultiUserServer/Source/MultiUserServer/Private/ConcertConsoleCommandExecutor.cpp
+++ b/Engine/Plugins/Developer/Concert/ConcertApp/MultiUserServer/Source/MultiUserServer/Private/ConcertConsoleCommandExecutor.cpp
@@ -86,6 +86,11 @@ namespace UE::MultiUserServer
 	{
 		return FInputChord();
 	}
+
+	FInputChord FConcertConsoleCommandExecutor::GetIterateExecutorHotKey() const
+	{
+		return FInputChord();
+	}
 }
 
 #undef LOCTEXT_NAMESPACE
diff --git a/Engine/Plugins/Developer/Concert/ConcertApp/MultiUserServer/Source/MultiUserServer/Private/ConcertConsoleCommandExecutor.h b/Engine/Plugins/Developer/Concert/ConcertApp/MultiUserServer/Source/MultiUserServer/Private/ConcertConsoleCommandExecutor.h
--- a/Engine/Plugins/Developer/Concert/ConcertApp/MultiUserServer/Source/MultiUserServer/Private/ConcertConsoleCommandExecutor.h
+++ b/Engine/Plugins/Developer/Concert/ConcertApp/MultiUserServer/Source/MultiUserServer/Private/ConcertConsoleCommandExecutor.h
@@ -25,6 +25,7 @@ namespace UE::MultiUserServer
 		virtual bool AllowHotKeyClose() const override;
 		virtual bool AllowMultiLine() const override;
 		virtual FInputChord GetHotKey() const override;
+		virtual FInputChord GetIterateExecutorHotKey() const override;
 		//~ End IConsoleCommandExecutor Interface
 	};
 }
diff --git a/Engine/Plugins/Experimental/PythonScriptPlugin/Source/PythonScriptPlugin/Private/PythonScriptPlugin.cpp b/Engine/Plugins/Experimental/PythonScriptPlugin/Source/PythonScriptPlugin/Private/PythonScriptPlugin.cpp
--- a/Engine/Plugins/Experimental/PythonScriptPlugin/Source/PythonScriptPlugin/Private/PythonScriptPlugin.cpp
+++ b/Engine/Plugins/Experimental/PythonScriptPlugin/Source/PythonScriptPlugin/Private/PythonScriptPlugin.cpp
@@ -170,6 +170,15 @@ FInputChord FPythonCommandExecutor::GetHotKey() const
 #endif
 }
 
+FInputChord FPythonCommandExecutor::GetIterateExecutorHotKey() const
+{
+#if WITH_EDITOR
+	return FGlobalEditorCommonCommands::Get().SelectNextConsoleExecutor->GetActiveChord(EMultipleKeyBindingIndex::Primary).Get();
+#else
+	return FInputChord();
+#endif
+}
+
 FPythonREPLCommandExecutor::FPythonREPLCommandExecutor(IPythonScriptPlugin* InPythonScriptPlugin)
 	: PythonScriptPlugin(InPythonScriptPlugin)
 {
@@ -243,6 +252,15 @@ FInputChord FPythonREPLCommandExecutor::GetHotKey() const
 #endif
 }
 
+FInputChord FPythonREPLCommandExecutor::GetIterateExecutorHotKey() const
+{
+#if WITH_EDITOR
+	return FGlobalEditorCommonCommands::Get().SelectNextConsoleExecutor->GetActiveChord(EMultipleKeyBindingIndex::Primary).Get();
+#else
+	return FInputChord();
+#endif
+}
+
 #if WITH_EDITOR
 class FPythonCommandMenuImpl : public IPythonCommandMenu
 {
diff --git a/Engine/Plugins/Experimental/PythonScriptPlugin/Source/PythonScriptPlugin/Private/PythonScriptPlugin.h b/Engine/Plugins/Experimental/PythonScriptPlugin/Source/PythonScriptPlugin/Private/PythonScriptPlugin.h
--- a/Engine/Plugins/Experimental/PythonScriptPlugin/Source/PythonScriptPlugin/Private/PythonScriptPlugin.h
+++ b/Engine/Plugins/Experimental/PythonScriptPlugin/Source/PythonScriptPlugin/Private/PythonScriptPlugin.h
@@ -45,6 +45,7 @@ public:
 	virtual bool AllowHotKeyClose() const override;
 	virtual bool AllowMultiLine() const override;
 	virtual FInputChord GetHotKey() const override;
+	virtual FInputChord GetIterateExecutorHotKey() const override;
 
 private:
 	IPythonScriptPlugin* PythonScriptPlugin;
@@ -69,6 +70,7 @@ public:
 	virtual bool AllowHotKeyClose() const override;
 	virtual bool AllowMultiLine() const override;
 	virtual FInputChord GetHotKey() const override;
+	virtual FInputChord GetIterateExecutorHotKey() const override;
 
 private:
 	IPythonScriptPlugin* PythonScriptPlugin;
diff --git a/Engine/Source/Developer/OutputLog/Private/SOutputLog.cpp b/Engine/Source/Developer/OutputLog/Private/SOutputLog.cpp
--- a/Engine/Source/Developer/OutputLog/Private/SOutputLog.cpp
+++ b/Engine/Source/Developer/OutputLog/Private/SOutputLog.cpp
@@ -459,6 +459,8 @@ FReply SConsoleInputBox::OnPreviewKeyDown(const FGeometry& MyGeometry, const FKe
 	}
 	else
 	{
+		const FInputChord KeyEventAsInputChord = FInputChord(KeyEvent.GetKey(), EModifierKey::FromBools(KeyEvent.IsControlDown(), KeyEvent.IsAltDown(), KeyEvent.IsShiftDown(), KeyEvent.IsCommandDown()));
+
 		if(KeyEvent.GetKey() == EKeys::Up)
 		{
 			// If the command field isn't empty we need you to have pressed Control+Up to summon the history (to make sure you're not just using caret navigation)
@@ -508,6 +510,11 @@ FReply SConsoleInputBox::OnPreviewKeyDown(const FGeometry& MyGeometry, const FKe
 
 			return FReply::Handled();
 		}
+		else if (ActiveCommandExecutor && ActiveCommandExecutor->GetIterateExecutorHotKey() == KeyEventAsInputChord)
+		{
+			MakeNextCommandExecutorActive();
+			return FReply::Handled();
+		}
 	}
 
 	return FReply::Unhandled();
@@ -660,6 +667,28 @@ bool SConsoleInputBox::IsCommandExecutorMenuEnabled() const
 	return !ConsoleCommandCustomExec.IsBound(); // custom execs always show the default executor in the UI (which has the selector disabled)
 }
 
+void SConsoleInputBox::MakeNextCommandExecutorActive()
+{
+	// Sorted so the iteration order matches the displayed order.
+	TArray<IConsoleCommandExecutor*> CommandExecutors = IModularFeatures::Get().GetModularFeatureImplementations<IConsoleCommandExecutor>(IConsoleCommandExecutor::ModularFeatureName());
+	CommandExecutors.Sort([](IConsoleCommandExecutor& LHS, IConsoleCommandExecutor& RHS)
+		{
+			return LHS.GetDisplayName().CompareTo(RHS.GetDisplayName()) < 0;
+		});
+
+	int32 CurrentIndex = CommandExecutors.IndexOfByKey(ActiveCommandExecutor);
+	if (CurrentIndex >= 0)
+	{
+		CurrentIndex++;
+		if (CurrentIndex >= CommandExecutors.Num())
+		{
+			CurrentIndex = 0;
+		}
+
+		SetActiveCommandExecutor(CommandExecutors[CurrentIndex]->GetName());
+	}
+}
+
 TSharedRef<SWidget> SConsoleInputBox::GetCommandExecutorMenuContent()
 {
 	FMenuBuilder MenuBuilder(/*bInShouldCloseWindowAfterMenuSelection=*/true, nullptr);
diff --git a/Engine/Source/Developer/OutputLog/Private/SOutputLog.h b/Engine/Source/Developer/OutputLog/Private/SOutputLog.h
--- a/Engine/Source/Developer/OutputLog/Private/SOutputLog.h
+++ b/Engine/Source/Developer/OutputLog/Private/SOutputLog.h
@@ -125,6 +125,8 @@ protected:
 
 	void SetActiveCommandExecutor(const FName InExecName);
 
+	void MakeNextCommandExecutorActive();
+
 	FText GetActiveCommandExecutorDisplayName() const;
 
 	FText GetActiveCommandExecutorHintText() const;
diff --git a/Engine/Source/Editor/UnrealEd/Private/Toolkits/GlobalEditorCommonCommands.cpp b/Engine/Source/Editor/UnrealEd/Private/Toolkits/GlobalEditorCommonCommands.cpp
--- a/Engine/Source/Editor/UnrealEd/Private/Toolkits/GlobalEditorCommonCommands.cpp
+++ b/Engine/Source/Editor/UnrealEd/Private/Toolkits/GlobalEditorCommonCommands.cpp
@@ -49,6 +49,7 @@ void FGlobalEditorCommonCommands::RegisterCommands()
 	UI_COMMAND(FindInContentBrowser, "Browse to Asset", "Browses to the associated asset and selects it in the most recently used Content Browser (summoning one if necessary)", EUserInterfaceActionType::Button, FInputChord(EModifierKey::Control, EKeys::B));
 
 	UI_COMMAND(OpenConsoleCommandBox, "Open Console Command Box", "Opens an edit box where you can type in a console command", EUserInterfaceActionType::Button, FInputChord(EKeys::Tilde));
+	UI_COMMAND(SelectNextConsoleExecutor, "Iterate Console Executor", "Iterates through active Console Executors (Python, Cmd, etc.)", EUserInterfaceActionType::Button, FInputChord(EModifierKey::Control, EKeys::Tilde));
 	
 	UI_COMMAND(OpenOutputLogDrawer, "Open Output Log Drawer", "Opens the output log drawer from the active asset editor status bar", EUserInterfaceActionType::Button, FInputChord(EModifierKey::Alt, EKeys::Tilde));
 
diff --git a/Engine/Source/Editor/UnrealEd/Public/Toolkits/FConsoleCommandExecutor.h b/Engine/Source/Editor/UnrealEd/Public/Toolkits/FConsoleCommandExecutor.h
--- a/Engine/Source/Editor/UnrealEd/Public/Toolkits/FConsoleCommandExecutor.h
+++ b/Engine/Source/Editor/UnrealEd/Public/Toolkits/FConsoleCommandExecutor.h
@@ -28,6 +28,11 @@ public:
 	{
 		return FGlobalEditorCommonCommands::Get().OpenConsoleCommandBox->GetActiveChord(EMultipleKeyBindingIndex::Primary).Get();
 	}
+
+	virtual FInputChord GetIterateExecutorHotKey() const override
+	{
+		return FGlobalEditorCommonCommands::Get().FGlobalEditorCommonCommands::Get().SelectNextConsoleExecutor->GetActiveChord(EMultipleKeyBindingIndex::Primary).Get();
+	}
 private:
 	bool ExecInternal(const TCHAR*) const;
 };
diff --git a/Engine/Source/Editor/UnrealEd/Public/Toolkits/GlobalEditorCommonCommands.h b/Engine/Source/Editor/UnrealEd/Public/Toolkits/GlobalEditorCommonCommands.h
--- a/Engine/Source/Editor/UnrealEd/Public/Toolkits/GlobalEditorCommonCommands.h
+++ b/Engine/Source/Editor/UnrealEd/Public/Toolkits/GlobalEditorCommonCommands.h
@@ -43,6 +43,7 @@ public:
 	TSharedPtr<FUICommandInfo> SummonOpenAssetDialogAlternate;
 	TSharedPtr<FUICommandInfo> OpenDocumentation;
 	TSharedPtr<FUICommandInfo> OpenConsoleCommandBox;
+	TSharedPtr<FUICommandInfo> SelectNextConsoleExecutor;
 	TSharedPtr<FUICommandInfo> OpenOutputLogDrawer;
 	TSharedPtr<FUICommandInfo> OpenContentBrowserDrawer;
 };
diff --git a/Engine/Source/Runtime/Core/Public/HAL/IConsoleManager.h b/Engine/Source/Runtime/Core/Public/HAL/IConsoleManager.h
--- a/Engine/Source/Runtime/Core/Public/HAL/IConsoleManager.h
+++ b/Engine/Source/Runtime/Core/Public/HAL/IConsoleManager.h
@@ -672,6 +672,11 @@ public:
 	* Returns the hotkey for this executor
 	*/
 	virtual struct FInputChord GetHotKey() const = 0;
+	/**
+	* Returns the hotkey to switch to the next executor. This is usually the same hotkey as GetHotKey but with a modifier,
+	* such as "`" turning into "Ctrl + `".
+	*/
+	virtual struct FInputChord GetIterateExecutorHotKey() const = 0;
 };
 
 
