Added new methods to SOutputLog.h and SOutputLog.cpp to create a new timestamp menu item in the output log settings menu.

diff --git a/Engine/Source/Developer/OutputLog/Private/SOutputLog.cpp b/Engine/Source/Developer/OutputLog/Private/SOutputLog.cpp
--- a/Engine/Source/Developer/OutputLog/Private/SOutputLog.cpp
+++ b/Engine/Source/Developer/OutputLog/Private/SOutputLog.cpp
@@ -1232,10 +1232,7 @@ void SOutputLog::Tick(const FGeometry& AllottedGeometry, const double InCurrentT
 	if (MessagesTextMarshaller->SubmitPendingMessages())
 	{
 		// Don't scroll to the bottom automatically when the user is scrolling the view or has scrolled it away from the bottom.
-		if (!bIsUserScrolled)
-		{
-			RequestForceScroll();
-		}
+		RequestForceScroll(true);
 	}
 
 	SCompoundWidget::Tick(AllottedGeometry, InCurrentTime, InDeltaTime);
@@ -1461,9 +1458,10 @@ void SOutputLog::OnConsoleCommandExecuted()
 	RequestForceScroll();
 }
 
-void SOutputLog::RequestForceScroll()
+void SOutputLog::RequestForceScroll(bool bIfUserHasNotScrolledUp)
 {
-	if (MessagesTextMarshaller->GetNumFilteredMessages() > 0)
+	if (MessagesTextMarshaller->GetNumFilteredMessages() > 0
+		&& (!bIfUserHasNotScrolledUp || !bIsUserScrolled))
 	{
 		MessagesTextBox->ScrollTo(ETextLocation::EndOfDocument);
 		bIsUserScrolled = false;
@@ -1497,10 +1495,59 @@ void SOutputLog::SetWordWrapEnabled(ECheckBoxState InValue)
 		Settings->SaveConfig();
 	}
 
-	if (!bIsUserScrolled)
+	RequestForceScroll(true);
+}
+
+ELogTimes::Type SOutputLog::GetSelectedTimestampMode() 
+{
+	const UOutputLogSettings* Settings = GetDefault<UOutputLogSettings>();
+	return Settings->LogTimestampMode;
+}
+
+bool SOutputLog::IsSelectedTimestampMode(ELogTimes::Type NewType)
+{
+	return GetSelectedTimestampMode() == NewType;
+}
+
+void SOutputLog::AddTimestampMenuSection(FMenuBuilder& Menu)
+{
+
+	Menu.BeginSection("LoggingTimestampSection");
+	{
+		const UEnum* Enum = StaticEnum<ELogTimes::Type>();
+
+		for (int CurrentTimeStampType = 0; CurrentTimeStampType < Enum->NumEnums() - 1; CurrentTimeStampType++)
+		{
+			
+			ELogTimes::Type TimeStampType = static_cast<ELogTimes::Type>(CurrentTimeStampType);
+			
+			Menu.AddMenuEntry(Enum->GetDisplayNameTextByIndex(CurrentTimeStampType),
+				Enum->GetToolTipTextByIndex(CurrentTimeStampType),
+				FSlateIcon(),
+				FUIAction(
+					FExecuteAction::CreateLambda([this, TimeStampType] {
+						SetTimestampMode(TimeStampType);
+						}),
+					FCanExecuteAction::CreateLambda([] { return true; }),
+					FIsActionChecked::CreateLambda([this, TimeStampType] { return IsSelectedTimestampMode(TimeStampType); })
+							),
+				NAME_None,
+				EUserInterfaceActionType::RadioButton);
+
+		}
+	}
+	Menu.EndSection();
+}
+
+void SOutputLog::SetTimestampMode(ELogTimes::Type InValue)
+{
+	 UOutputLogSettings* Settings = GetMutableDefault<UOutputLogSettings>();
+	if (Settings)
 	{
-		RequestForceScroll();
+		Settings->LogTimestampMode = InValue;
+		Settings->SaveConfig();
 	}
+	RequestForceScroll(true);
 }
 
 #if WITH_EDITOR
@@ -1800,6 +1847,19 @@ TSharedRef<SWidget> SOutputLog::GetViewButtonContent(EOutputLogSettingsMenuFlags
 		);
 	}
 
+	FProperty* MetadataProperty = UOutputLogSettings::StaticClass()->FindPropertyByName(GET_MEMBER_NAME_CHECKED(UOutputLogSettings, LogTimestampMode));
+	
+	MenuBuilder.AddSubMenu(
+		TAttribute<FText>::CreateLambda([this]()
+			{
+				const UEnum* Enum = StaticEnum<ELogTimes::Type>();
+				return FText::Format(LOCTEXT("TimestampsSubmenu", "Timestamp Mode: {0}"),
+					Enum->GetDisplayNameTextByIndex(this->GetSelectedTimestampMode()));
+			}),
+		MetadataProperty->GetToolTipText(),
+		FNewMenuDelegate::CreateSP(this, &SOutputLog::AddTimestampMenuSection));
+
+
 #if WITH_EDITOR
 	const bool bSupportClearOnPie = (Flags & EOutputLogSettingsMenuFlags::SkipClearOnPie) == EOutputLogSettingsMenuFlags::None;
 	if (bSupportClearOnPie)
diff --git a/Engine/Source/Developer/OutputLog/Private/SOutputLog.h b/Engine/Source/Developer/OutputLog/Private/SOutputLog.h
--- a/Engine/Source/Developer/OutputLog/Private/SOutputLog.h
+++ b/Engine/Source/Developer/OutputLog/Private/SOutputLog.h
@@ -396,7 +396,7 @@ protected:
 	void OnConsoleCommandExecuted();
 
 	/** Request we immediately force scroll to the bottom of the log */
-	void RequestForceScroll();
+	void RequestForceScroll(bool bIfUserHasNotScrolledUp = false);
 
 	/** Converts the array of messages into something the text box understands */
 	TSharedPtr< FOutputLogTextLayoutMarshaller > MessagesTextMarshaller;
@@ -463,6 +463,14 @@ private:
 
 	void SetWordWrapEnabled(ECheckBoxState InValue);
 
+	void SetTimestampMode(ELogTimes::Type InValue);
+
+	bool IsSelectedTimestampMode(ELogTimes::Type NewType);
+
+	void AddTimestampMenuSection(FMenuBuilder& Menu);
+
+	ELogTimes::Type GetSelectedTimestampMode();
+
 #if WITH_EDITOR
 	bool IsClearOnPIEEnabled() const;
 
