Altering SOutputLog.cpp to fix compilation issues

diff --git a/Engine/Source/Developer/OutputLog/Private/SOutputLog.cpp b/Engine/Source/Developer/OutputLog/Private/SOutputLog.cpp
--- a/Engine/Source/Developer/OutputLog/Private/SOutputLog.cpp
+++ b/Engine/Source/Developer/OutputLog/Private/SOutputLog.cpp
@@ -1520,9 +1520,14 @@ void SOutputLog::AddTimestampMenuSection(FMenuBuilder& Menu)
 		{
 			
 			ELogTimes::Type TimeStampType = static_cast<ELogTimes::Type>(CurrentTimeStampType);
+			FText Tooltip;
+
+			#if WITH_EDITOR
+				Tooltip = Enum->GetToolTipTextByIndex(CurrentTimeStampType);
+			#endif // WITH_EDITOR
 			
 			Menu.AddMenuEntry(Enum->GetDisplayNameTextByIndex(CurrentTimeStampType),
-				Enum->GetToolTipTextByIndex(CurrentTimeStampType),
+				Tooltip,
 				FSlateIcon(),
 				FUIAction(
 					FExecuteAction::CreateLambda([this, TimeStampType] {
@@ -1846,9 +1851,12 @@ TSharedRef<SWidget> SOutputLog::GetViewButtonContent(EOutputLogSettingsMenuFlags
 			EUserInterfaceActionType::ToggleButton
 		);
 	}
+	FText Tooltip;
 
-	FProperty* MetadataProperty = UOutputLogSettings::StaticClass()->FindPropertyByName(GET_MEMBER_NAME_CHECKED(UOutputLogSettings, LogTimestampMode));
-	
+	#if WITH_EDITORONLY_DATA
+		Tooltip = UOutputLogSettings::StaticClass()->FindPropertyByName(
+			GET_MEMBER_NAME_CHECKED(UOutputLogSettings, LogTimestampMode))->GetToolTipText();
+	#endif // WITH_EDITORONLY_DATA
 	MenuBuilder.AddSubMenu(
 		TAttribute<FText>::CreateLambda([this]()
 			{
@@ -1856,7 +1864,7 @@ TSharedRef<SWidget> SOutputLog::GetViewButtonContent(EOutputLogSettingsMenuFlags
 				return FText::Format(LOCTEXT("TimestampsSubmenu", "Timestamp Mode: {0}"),
 					Enum->GetDisplayNameTextByIndex(this->GetSelectedTimestampMode()));
 			}),
-		MetadataProperty->GetToolTipText(),
+		Tooltip,
 		FNewMenuDelegate::CreateSP(this, &SOutputLog::AddTimestampMenuSection));
 
 
