Possible compile fix for OutputLog found on the release line.

diff --git a/Engine/Source/Developer/OutputLog/OutputLog.Build.cs b/Engine/Source/Developer/OutputLog/OutputLog.Build.cs
--- a/Engine/Source/Developer/OutputLog/OutputLog.Build.cs
+++ b/Engine/Source/Developer/OutputLog/OutputLog.Build.cs
@@ -13,22 +13,25 @@ public class OutputLog : ModuleRules
                 "InputCore",
 				"Slate",
 				"SlateCore",
-                "TargetPlatform",
-                "DesktopPlatform",
+				"TargetPlatform",
+				"DesktopPlatform",
 				"ToolWidgets", 
 			}
 		);
 
 		if (Target.bBuildEditor)
-        {
-            PrivateDependencyModuleNames.AddRange(
+		{
+			PrivateDependencyModuleNames.AddRange(
 				new string[] {
 					"EditorFramework",
 					"StatusBar",
 					"UnrealEd",
 				}
 			);
+		}
 
+		if (Target.bBuildEditor || Target.bBuildDeveloperTools)
+		{
 			PrivateIncludePathModuleNames.AddRange(
 				new string[] {
 					"WorkspaceMenuStructure",
