Editor: DeviceOutputLog: Ensured the right destrubtion order between ITargetDeviceOutput and its owning SDeviceOutputLog.

Previously, ITargetDeviceOutput was automatically destroyed as part of destruction of SDeviceOutputLog. It means that if ITargetDeviceOutput tried to deference the owner during its own destruction it could potentially dereference members that have already been destroyed.

Resolved by manually clearing CurrentDeviceOutputPtr which triggers destruction of ITargetDeviceOutput a bit sooner when all the owners members are still valid.

#rnx

diff --git a/Engine/Source/Developer/OutputLog/Private/SDeviceOutputLog.cpp b/Engine/Source/Developer/OutputLog/Private/SDeviceOutputLog.cpp
--- a/Engine/Source/Developer/OutputLog/Private/SDeviceOutputLog.cpp
+++ b/Engine/Source/Developer/OutputLog/Private/SDeviceOutputLog.cpp
@@ -129,6 +129,10 @@ SDeviceOutputLog::~SDeviceOutputLog()
 {
 	ITargetPlatform::OnDeviceDiscovered().RemoveAll(this);
 	ITargetPlatform::OnDeviceLost().RemoveAll(this);
+
+	// Clearing the pointer manually to ensure that when the pointed device output object is destroyed
+	// SDeviceOutputLog is still in a valid state in case CurrentDeviceOutputPtr wanted to dereference it.
+	CurrentDeviceOutputPtr.Reset();
 }
 
 void SDeviceOutputLog::Tick(const FGeometry& AllottedGeometry, const double InCurrentTime, const float InDeltaTime)
