Prevent trying to remove output device while in a signal handling from this callback which is no longer logging during a crash anymore

diff --git a/Engine/Source/Developer/OutputLog/Private/SOutputLog.cpp b/Engine/Source/Developer/OutputLog/Private/SOutputLog.cpp
--- a/Engine/Source/Developer/OutputLog/Private/SOutputLog.cpp
+++ b/Engine/Source/Developer/OutputLog/Private/SOutputLog.cpp
@@ -1195,8 +1195,6 @@ void SOutputLog::Construct( const FArguments& InArgs, bool bCreateDrawerDockButt
 	];
 
 	GLog->AddOutputDevice(this);
-	// Remove itself on crash (crashmalloc has limited memory and echoing logs here at that point is useless).
-	FCoreDelegates::OnHandleSystemError.AddRaw(this, &SOutputLog::OnCrash);
 
 #if WITH_EDITOR
 	// Listen for style changes
@@ -1238,14 +1236,6 @@ void SOutputLog::Tick(const FGeometry& AllottedGeometry, const double InCurrentT
 	SCompoundWidget::Tick(AllottedGeometry, InCurrentTime, InDeltaTime);
 }
 
-void SOutputLog::OnCrash()
-{
-	if (GLog != nullptr)
-	{
-		GLog->RemoveOutputDevice(this);
-	}
-}
-
 static const FName NAME_StyleLogCommand(TEXT("Log.Command"));
 static const FName NAME_StyleLogError(TEXT("Log.Error"));
 static const FName NAME_StyleLogWarning(TEXT("Log.Warning"));
diff --git a/Engine/Source/Developer/OutputLog/Private/SOutputLog.h b/Engine/Source/Developer/OutputLog/Private/SOutputLog.h
--- a/Engine/Source/Developer/OutputLog/Private/SOutputLog.h
+++ b/Engine/Source/Developer/OutputLog/Private/SOutputLog.h
@@ -378,9 +378,6 @@ protected:
 
 	virtual void Serialize( const TCHAR* V, ELogVerbosity::Type Verbosity, const class FName& Category ) override;
 
-	/* Remove itself on crash to prevent adding log lines here */
-	void OnCrash();
-
 protected:
 	/**
 	 * Extends the context menu used by the text box
