Speculative fix crash when packaging due to logging.

diff --git a/Engine/Source/Developer/OutputLog/Private/SOutputLog.cpp b/Engine/Source/Developer/OutputLog/Private/SOutputLog.cpp
--- a/Engine/Source/Developer/OutputLog/Private/SOutputLog.cpp
+++ b/Engine/Source/Developer/OutputLog/Private/SOutputLog.cpp
@@ -823,6 +823,27 @@ bool FOutputLogTextLayoutMarshaller::SubmitPendingMessages()
 	if (Messages.IsValidIndex(NextPendingMessageIndex))
 	{
 		const int32 CurrentMessagesCount = Messages.Num();
+
+		// If the same message is still invalid after a tick, just skip it during processing (Prevents getting stuck on a message).
+		bool bMessageStillInvalidAfterTick = Messages.IsValidIndex(LastNullMessageIndex) && !Messages[LastNullMessageIndex];
+
+		if (!bMessageStillInvalidAfterTick)
+		{
+			// Text log can be written to async. Rather than sync the log. Process all messages if possible,
+			// but if not wait till next tick to process (For the most part we will always process).
+			for (int32 MessageIndex = NextPendingMessageIndex; MessageIndex < CurrentMessagesCount; ++MessageIndex)
+			{
+				const TSharedPtr<FOutputLogMessage> Message = Messages[MessageIndex];
+				if (!Message)
+				{
+					LastNullMessageIndex = MessageIndex;
+					return false;
+				}
+			}
+		}
+
+		LastNullMessageIndex = INDEX_NONE;
+
 		AppendPendingMessagesToTextLayout();
 		NextPendingMessageIndex = CurrentMessagesCount;
 		return true;
@@ -895,6 +916,11 @@ void FOutputLogTextLayoutMarshaller::AppendPendingMessagesToTextLayout()
 		const TSharedPtr<FOutputLogMessage> Message = Messages[MessageIndex];
 		const int32 LineIndex = TextLayout->GetLineModels().Num() + NumAddedMessages;
 
+		if (!Message)
+		{
+			continue;
+		}
+
 		Filter->AddAvailableLogCategory(Message->Category);
 		if (!Filter->IsMessageAllowed(Message))
 		{
@@ -1064,6 +1090,7 @@ FTextLocation FOutputLogTextLayoutMarshaller::GetTextLocationAt(const FVector2D&
 FOutputLogTextLayoutMarshaller::FOutputLogTextLayoutMarshaller(TArray< TSharedPtr<FOutputLogMessage> > InMessages, FOutputLogFilter* InFilter)
 	: Messages(MoveTemp(InMessages))
 	, NextPendingMessageIndex(0)
+	, LastNullMessageIndex(INDEX_NONE)
 	, CachedNumMessages(0)
 	, Filter(InFilter)
 	, TextLayout(nullptr)
diff --git a/Engine/Source/Developer/OutputLog/Private/SOutputLog.h b/Engine/Source/Developer/OutputLog/Private/SOutputLog.h
--- a/Engine/Source/Developer/OutputLog/Private/SOutputLog.h
+++ b/Engine/Source/Developer/OutputLog/Private/SOutputLog.h
@@ -544,6 +544,9 @@ protected:
 	/** Index of the next entry in the Messages array that is pending submission to the text layout */
 	int32 NextPendingMessageIndex;
 
+	/** Index of last message that was null, used to ensure we don't stop processing due to a null message */
+	int32 LastNullMessageIndex;
+
 	/** Holds cached numbers of messages to avoid unnecessary re-filtering */
 	int32 CachedNumMessages;
 	
