Add Locks to SOutputLog to prevent read / write race condition crashes.

diff --git a/Engine/Source/Developer/OutputLog/Private/SOutputLog.cpp b/Engine/Source/Developer/OutputLog/Private/SOutputLog.cpp
--- a/Engine/Source/Developer/OutputLog/Private/SOutputLog.cpp
+++ b/Engine/Source/Developer/OutputLog/Private/SOutputLog.cpp
@@ -5,6 +5,7 @@
 #include "Framework/Text/TextLayout.h"
 #include "Misc/ConfigCacheIni.h"
 #include "Misc/OutputDeviceHelper.h"
+#include "Misc/ScopeLock.h"
 #include "SlateOptMacros.h"
 #include "Textures/SlateIcon.h"
 #include "Framework/Commands/UIAction.h"
@@ -815,39 +816,34 @@ void FOutputLogTextLayoutMarshaller::GetText(FString& TargetString, const FTextL
 
 bool FOutputLogTextLayoutMarshaller::AppendPendingMessage(const TCHAR* InText, const ELogVerbosity::Type InVerbosity, const FName& InCategory)
 {
-	return SOutputLog::CreateLogMessages(InText, InVerbosity, InCategory, Messages);
+	// We don't want to skip adding messages, so just try to acquire the lock
+	FScopeLock PendingMessagesAccess(&PendingMessagesCriticalSection);
+	return SOutputLog::CreateLogMessages(InText, InVerbosity, InCategory, PendingMessages);
 }
 
 bool FOutputLogTextLayoutMarshaller::SubmitPendingMessages()
 {
+	// We can always submit messages next tick. So only try to lock, if not possible return.
+	if (PendingMessagesCriticalSection.TryLock())
+	{
+		Messages.Append(MoveTemp(PendingMessages));
+		PendingMessages.Reset();
+		PendingMessagesCriticalSection.Unlock();
+	}
+	else
+	{
+		return false;
+	}
+
 	if (Messages.IsValidIndex(NextPendingMessageIndex))
 	{
 		const int32 CurrentMessagesCount = Messages.Num();
 
-		// If the same message is still invalid after a tick, just skip it during processing (Prevents getting stuck on a message).
-		bool bMessageStillInvalidAfterTick = Messages.IsValidIndex(LastNullMessageIndex) && !Messages[LastNullMessageIndex];
-
-		if (!bMessageStillInvalidAfterTick)
-		{
-			// Text log can be written to async. Rather than sync the log. Process all messages if possible,
-			// but if not wait till next tick to process (For the most part we will always process).
-			for (int32 MessageIndex = NextPendingMessageIndex; MessageIndex < CurrentMessagesCount; ++MessageIndex)
-			{
-				const TSharedPtr<FOutputLogMessage> Message = Messages[MessageIndex];
-				if (!Message)
-				{
-					LastNullMessageIndex = MessageIndex;
-					return false;
-				}
-			}
-		}
-
-		LastNullMessageIndex = INDEX_NONE;
-
 		AppendPendingMessagesToTextLayout();
 		NextPendingMessageIndex = CurrentMessagesCount;
 		return true;
 	}
+
 	return false;
 }
 
@@ -1090,7 +1086,6 @@ FTextLocation FOutputLogTextLayoutMarshaller::GetTextLocationAt(const FVector2D&
 FOutputLogTextLayoutMarshaller::FOutputLogTextLayoutMarshaller(TArray< TSharedPtr<FOutputLogMessage> > InMessages, FOutputLogFilter* InFilter)
 	: Messages(MoveTemp(InMessages))
 	, NextPendingMessageIndex(0)
-	, LastNullMessageIndex(INDEX_NONE)
 	, CachedNumMessages(0)
 	, Filter(InFilter)
 	, TextLayout(nullptr)
diff --git a/Engine/Source/Developer/OutputLog/Private/SOutputLog.h b/Engine/Source/Developer/OutputLog/Private/SOutputLog.h
--- a/Engine/Source/Developer/OutputLog/Private/SOutputLog.h
+++ b/Engine/Source/Developer/OutputLog/Private/SOutputLog.h
@@ -14,6 +14,7 @@
 #include "Widgets/Views/STableRow.h"
 #include "Framework/Text/BaseTextLayoutMarshaller.h"
 #include "Misc/TextFilterExpressionEvaluator.h"
+#include "HAL/CriticalSection.h"
 #include "HAL/IConsoleManager.h"
 
 class FMenuBuilder;
@@ -541,12 +542,12 @@ protected:
 	/** All log messages to show in the text box */
 	TArray< TSharedPtr<FOutputLogMessage> > Messages;
 
+	/** Messages pending add, kept separate to avoid a race condition when reading Messages */
+	TArray< TSharedPtr<FOutputLogMessage> > PendingMessages;
+
 	/** Index of the next entry in the Messages array that is pending submission to the text layout */
 	int32 NextPendingMessageIndex;
 
-	/** Index of last message that was null, used to ensure we don't stop processing due to a null message */
-	int32 LastNullMessageIndex;
-
 	/** Holds cached numbers of messages to avoid unnecessary re-filtering */
 	int32 CachedNumMessages;
 	
@@ -559,4 +560,7 @@ protected:
 	FName CategoryToHighlight;
 
 	FTextLayout* TextLayout;
+
+	/** Output log runs own its own "OutputDeviceRedirector" thread, lock against messages to prevent race conditions */
+	FCriticalSection PendingMessagesCriticalSection;
 };
