Code only cvar to hide console command

diff --git a/Engine/Source/Developer/OutputLog/Private/OutputLogModule.cpp b/Engine/Source/Developer/OutputLog/Private/OutputLogModule.cpp
--- a/Engine/Source/Developer/OutputLog/Private/OutputLogModule.cpp
+++ b/Engine/Source/Developer/OutputLog/Private/OutputLogModule.cpp
@@ -38,6 +38,20 @@ namespace OutputLogModule
 {
 	static const FName OutputLogTabName = FName(TEXT("OutputLog"));
 	static const FName DeviceOutputLogTabName = FName(TEXT("DeviceOutputLog"));
+
+	bool bHideConsole = false;
+	FAutoConsoleVariableRef CVarHideConsoleCommand(
+		TEXT("OutputLogModule.HideConsole"), 
+		bHideConsole, 
+		TEXT("Whether debug console widgets should be hidden (false by default)"), 
+		FConsoleVariableDelegate::CreateLambda([](IConsoleVariable* /*CVar*/)
+			{
+				if (bHideConsole)
+				{
+					FOutputLogModule::Get().CloseDebugConsole();
+				}
+			}),
+		ECVF_ReadOnly);
 }
 
 /** This class is to capture all log output even if the log window is closed */
@@ -168,10 +182,16 @@ FOutputLogModule& FOutputLogModule::Get()
 	return FModuleManager::Get().LoadModuleChecked<FOutputLogModule>(OutputLog);
 }
 
+bool FOutputLogModule::ShouldHideConsole() const
+{
+	return OutputLogModule::bHideConsole;
+}
+
 TSharedRef<SWidget> FOutputLogModule::MakeConsoleInputBox(TSharedPtr<SMultiLineEditableTextBox>& OutExposedEditableTextBox, const FSimpleDelegate& OnCloseConsole, const FSimpleDelegate& OnConsoleCommandExecuted) const
 {
 	TSharedRef<SConsoleInputBox> NewConsoleInputBox =
 		SNew(SConsoleInputBox)
+		.Visibility(MakeAttributeLambda([](){ return FOutputLogModule::Get().ShouldHideConsole() ? EVisibility::Collapsed : EVisibility::Visible; }))
 		.OnCloseConsole(OnCloseConsole)
 		.OnConsoleCommandExecuted(OnConsoleCommandExecuted);
 
@@ -207,6 +227,11 @@ TSharedRef<SWidget> FOutputLogModule::MakeOutputLogWidget(const FOutputLogCreati
 
 void FOutputLogModule::ToggleDebugConsoleForWindow(const TSharedRef<SWindow>& Window, const EDebugConsoleStyle::Type InStyle, const FDebugConsoleDelegates& DebugConsoleDelegates)
 {
+	if (ShouldHideConsole())
+	{
+		return;
+	}
+
 	bool bShouldOpen = true;
 	// Close an existing console box, if there is one
 	TSharedPtr< SWidget > PinnedDebugConsole(DebugConsole.Pin());
diff --git a/Engine/Source/Developer/OutputLog/Private/SDeviceOutputLog.cpp b/Engine/Source/Developer/OutputLog/Private/SDeviceOutputLog.cpp
--- a/Engine/Source/Developer/OutputLog/Private/SDeviceOutputLog.cpp
+++ b/Engine/Source/Developer/OutputLog/Private/SDeviceOutputLog.cpp
@@ -12,6 +12,7 @@
 #include "Interfaces/ITargetPlatform.h"
 #include "Interfaces/ITargetPlatformManagerModule.h"
 #include "PlatformInfo.h"
+#include "OutputLogModule.h"
 #include "OutputLogStyle.h"
 #include "SSimpleComboButton.h"
 
@@ -91,6 +92,7 @@ void SDeviceOutputLog::Construct( const FArguments& InArgs )
 				.VAlign(VAlign_Center)
 				[
 					SNew(SConsoleInputBox)
+					.Visibility(MakeAttributeLambda([]() { return FOutputLogModule::Get().ShouldHideConsole() ? EVisibility::Collapsed : EVisibility::Visible; }))
 					.ConsoleCommandCustomExec(this, &SDeviceOutputLog::ExecuteConsoleCommand)
 					.OnConsoleCommandExecuted(this, &SDeviceOutputLog::OnConsoleCommandExecuted)
 					// Always place suggestions above the input line for the output log widget
diff --git a/Engine/Source/Developer/OutputLog/Private/SOutputLog.cpp b/Engine/Source/Developer/OutputLog/Private/SOutputLog.cpp
--- a/Engine/Source/Developer/OutputLog/Private/SOutputLog.cpp
+++ b/Engine/Source/Developer/OutputLog/Private/SOutputLog.cpp
@@ -1208,6 +1208,7 @@ void SOutputLog::Construct( const FArguments& InArgs, bool bCreateDrawerDockButt
 		.AutoHeight()
 		[
 			SAssignNew(ConsoleInputBox, SConsoleInputBox)
+			.Visibility(MakeAttributeLambda([]() { return  FOutputLogModule::Get().ShouldHideConsole() ? EVisibility::Collapsed : EVisibility::Visible; }))
 			.OnConsoleCommandExecuted(this, &SOutputLog::OnConsoleCommandExecuted)
 			.OnCloseConsole(InArgs._OnCloseConsole)
 			// Always place suggestions above the input line for the output log widget
diff --git a/Engine/Source/Developer/OutputLog/Public/OutputLogModule.h b/Engine/Source/Developer/OutputLog/Public/OutputLogModule.h
--- a/Engine/Source/Developer/OutputLog/Public/OutputLogModule.h
+++ b/Engine/Source/Developer/OutputLog/Public/OutputLogModule.h
@@ -42,6 +42,9 @@ public:
 
 	static OUTPUTLOG_API FOutputLogModule& Get();
 
+	/** Returns whether debug console widgets should be hidden */
+	virtual bool ShouldHideConsole() const;
+
 	/** Generates a console input box widget.  Remember, this widget will become invalid if the
 		output log DLL is unloaded on the fly. */
 	virtual TSharedRef<SWidget> MakeConsoleInputBox(TSharedPtr<SMultiLineEditableTextBox>& OutExposedEditableTextBox, const FSimpleDelegate& OnCloseConsole, const FSimpleDelegate& OnConsoleCommandExecuted) const;
diff --git a/Engine/Source/Editor/StatusBar/Private/StatusBarSubsystem.cpp b/Engine/Source/Editor/StatusBar/Private/StatusBarSubsystem.cpp
--- a/Engine/Source/Editor/StatusBar/Private/StatusBarSubsystem.cpp
+++ b/Engine/Source/Editor/StatusBar/Private/StatusBarSubsystem.cpp
@@ -26,6 +26,7 @@
 #include "Widgets/Input/SMultiLineEditableTextBox.h"
 #include "OutputLogSettings.h"
 #include "SOneTimeIndustryQuery.h"
+#include "Types/SlateAttributeMetaData.h"
 
 #define LOCTEXT_NAMESPACE "StatusBar"
 
@@ -462,19 +463,37 @@ TSharedRef<SWidget> UStatusBarSubsystem::MakeStatusBarWidget(FName StatusBarName
 	FSimpleDelegate OnConsoleCommandExecuted;
 
 	TSharedPtr<SMultiLineEditableTextBox> ConsoleEditBox;
-	TSharedRef<SWidget> OutputLog = 
+	TSharedPtr<SWidget> OutputLog;
+	{
+		TSharedRef<SWidget> ConsoleInputBox = OutputLogModule.MakeConsoleInputBox(ConsoleEditBox, OnConsoleClosed, OnConsoleCommandExecuted);
+
+		auto IsConsoleInputBoxBorderVisible = [ConsoleInputBoxWeak = ConsoleInputBox.ToWeakPtr()]()
+		{
+			if (TSharedPtr<SWidget> ConsoleInputBox = ConsoleInputBoxWeak.Pin())
+			{
+				FSlateAttributeMetaData::UpdateOnlyVisibilityAttributes(*ConsoleInputBox, FSlateAttributeMetaData::EInvalidationPermission::AllowInvalidationIfConstructed);
+				if (ConsoleInputBox->GetVisibility() == EVisibility::Collapsed)
+				{
+					return EVisibility::Collapsed;
+				}
+			}
+			return EVisibility::Visible;
+		};
+
+		OutputLog =
 			SNew(SBorder)
 			.BorderImage(FAppStyle::Get().GetBrush("Brushes.Panel"))
 			.VAlign(VAlign_Center)
 			.Padding(FMargin(6.0f, 0.0f))
+			.Visibility(MakeAttributeLambda(IsConsoleInputBoxBorderVisible))
 			[
 				SNew(SBox)
 				.WidthOverride(350.f)
 				[
-					OutputLogModule.MakeConsoleInputBox(ConsoleEditBox, OnConsoleClosed, OnConsoleCommandExecuted)
+					ConsoleInputBox
 				]
 			];
-	
+	}
 
 	FWidgetDrawerConfig OutputLogDrawer(StatusBarDrawerIds::OutputLog);
 
