Add a caching mechanism to OutputLogModule to persist per session filter settings when the OutputLog is closed then reopened

diff --git a/Engine/Source/Developer/OutputLog/Private/OutputLogModule.cpp b/Engine/Source/Developer/OutputLog/Private/OutputLogModule.cpp
--- a/Engine/Source/Developer/OutputLog/Private/OutputLogModule.cpp
+++ b/Engine/Source/Developer/OutputLog/Private/OutputLogModule.cpp
@@ -98,6 +98,7 @@ private:
 TSharedRef<SDockTab> FOutputLogModule::SpawnOutputLogTab(const FSpawnTabArgs& Args)
 {
 	TSharedRef<SOutputLog> NewLog = SNew(SOutputLog, false).Messages(OutputLogHistory->GetMessages());
+	NewLog->UpdateOutputLogFilter(OutputLogFilterCache->GetSelectedLogCategories(), OutputLogFilterCache->bShowErrors, OutputLogFilterCache->bShowWarnings, OutputLogFilterCache->bShowLogs);
 
 	OutputLog = NewLog;
 
@@ -108,11 +109,22 @@ TSharedRef<SDockTab> FOutputLogModule::SpawnOutputLogTab(const FSpawnTabArgs& Ar
 			NewLog
 		];
 
+	NewTab->SetOnTabClosed(SDockTab::FOnTabClosedCallback::CreateRaw(this, &FOutputLogModule::OnOutputLogTabClosed));
+
 	OutputLogTab = NewTab;
 
 	return NewTab;
 }
 
+void FOutputLogModule::OnOutputLogTabClosed(TSharedRef<SDockTab> Tab)
+{
+	if (TSharedPtr<SOutputLog> SharedOutputLog = OutputLog.Pin())
+	{
+		// Cache the closing LogFilterTab so that we can restore the same filter when it's opened again
+		*OutputLogFilterCache = SharedOutputLog->GetOutputLogFilter();
+	}
+}
+
 TSharedRef<SDockTab> FOutputLogModule::SpawnDeviceOutputLogTab(const FSpawnTabArgs& Args)
 {
 	return SNew(SDockTab)
@@ -156,6 +168,7 @@ void FOutputLogModule::StartupModule()
 #endif
 
 	OutputLogHistory = MakeShareable(new FOutputLogHistory);
+	OutputLogFilterCache = MakeUnique<FOutputLogFilter>();
 }
 
 void FOutputLogModule::ShutdownModule()
@@ -209,6 +222,9 @@ TSharedRef<SWidget> FOutputLogModule::MakeOutputLogDrawerWidget(const FSimpleDel
 			SNew(SOutputLog, true)
 			.OnCloseConsole(OnCloseConsole)
 			.Messages(OutputLogHistory->GetMessages());
+
+		OutputLogDrawerPinned->UpdateOutputLogFilter(OutputLogFilterCache->GetSelectedLogCategories(), OutputLogFilterCache->bShowErrors, OutputLogFilterCache->bShowWarnings, OutputLogFilterCache->bShowLogs);
+
 		OutputLogDrawer = OutputLogDrawerPinned;
 	}
 
@@ -391,6 +407,32 @@ void FOutputLogModule::UpdateOutputLogFilter(const TArray<FName>& CategoriesToSh
 	{
 		SharedOutputLog->UpdateOutputLogFilter(CategoriesToShow, bShowErrors, bShowWarnings, bShowLogs);
 	}
+
+	if (TSharedPtr<SOutputLog> SharedOutputLogDrawer = OutputLogDrawer.Pin())
+	{
+		SharedOutputLogDrawer->UpdateOutputLogFilter(CategoriesToShow, bShowErrors, bShowWarnings, bShowLogs);
+	}
+	
+	// Update the filter cache to these new settings
+	// This will be useful for the case where the OutputLog Drawer or Tab get created after this call
+	if (bShowErrors.IsSet())
+	{
+		OutputLogFilterCache->bShowErrors = bShowErrors.GetValue();
+	}
+	if (bShowWarnings.IsSet())
+	{
+		OutputLogFilterCache->bShowWarnings = bShowWarnings.GetValue();
+	}
+	if (bShowLogs.IsSet())
+	{
+		OutputLogFilterCache->bShowLogs = bShowLogs.GetValue();
+	}
+
+	OutputLogFilterCache->ClearSelectedLogCategories();
+	for (const FName& CategoryToShow : CategoriesToShow)
+	{
+		OutputLogFilterCache->ToggleLogCategory(CategoryToShow);
+	}
 }
 
 void FOutputLogModule::OpenOutputLog() const
diff --git a/Engine/Source/Developer/OutputLog/Private/SOutputLog.h b/Engine/Source/Developer/OutputLog/Private/SOutputLog.h
--- a/Engine/Source/Developer/OutputLog/Private/SOutputLog.h
+++ b/Engine/Source/Developer/OutputLog/Private/SOutputLog.h
@@ -274,6 +274,8 @@ struct FOutputLogFilter
 
 	const TArray<FName>& GetAvailableLogCategories() { return AvailableLogCategories; }
 
+	const TArray<FName>& GetSelectedLogCategories() { return SelectedLogCategories; }
+
 	/** Adds a Log Category to the list of available categories, if it isn't already present */
 	void AddAvailableLogCategory(const FName& LogCategory);
 
@@ -310,7 +312,7 @@ public:
 	SLATE_BEGIN_ARGS( SOutputLog )
 		: _Messages()
 		{}
-		
+
 		SLATE_EVENT(FSimpleDelegate, OnCloseConsole)
 
 		/** All messages captured before this log window has been created */
@@ -375,6 +377,8 @@ public:
 
 	/** Change the output log's filter. If CategoriesToShow is empty, all categories will be shown. */
 	void UpdateOutputLogFilter(const TArray<FName>& CategoriesToShow, TOptional<bool> bShowErrors = TOptional<bool>(), TOptional<bool> bShowWarnings = TOptional<bool>(), TOptional<bool> bShowLogs = TOptional<bool>());
+
+	const FOutputLogFilter& GetOutputLogFilter() { return Filter; }
 protected:
 
 	virtual void Serialize( const TCHAR* V, ELogVerbosity::Type Verbosity, const class FName& Category ) override;
diff --git a/Engine/Source/Developer/OutputLog/Public/OutputLogModule.h b/Engine/Source/Developer/OutputLog/Public/OutputLogModule.h
--- a/Engine/Source/Developer/OutputLog/Public/OutputLogModule.h
+++ b/Engine/Source/Developer/OutputLog/Public/OutputLogModule.h
@@ -13,6 +13,7 @@ class SOutputLog;
 class SDockTab;
 class FSpawnTabArgs;
 struct FOutputLogCreationParams;
+struct FOutputLogFilter;
 
 /** Style of the debug console */
 namespace EDebugConsoleStyle
@@ -81,12 +82,17 @@ public:
 
 private:
 	TSharedRef<SDockTab> SpawnOutputLogTab(const FSpawnTabArgs& Args);
+	void OnOutputLogTabClosed(TSharedRef<SDockTab> Tab);
+
 	TSharedRef<SDockTab> SpawnDeviceOutputLogTab(const FSpawnTabArgs& Args);
 
 private:
 	/** Our global output log app spawner */
 	TSharedPtr<FOutputLogHistory> OutputLogHistory;
 
+	/** Caches the user selected Filters as the OutputLog tab can be closed and remade multiple times */
+	TUniquePtr<FOutputLogFilter> OutputLogFilterCache;
+
 	/** Our global active output log that belongs to a tab */
 	TWeakPtr<SOutputLog> OutputLog;
 
