SOutputLog: remove no-longer-needed deregistration now that we have removed the registration.
#rnx

diff --git a/Engine/Source/Developer/OutputLog/Private/SOutputLog.cpp b/Engine/Source/Developer/OutputLog/Private/SOutputLog.cpp
--- a/Engine/Source/Developer/OutputLog/Private/SOutputLog.cpp
+++ b/Engine/Source/Developer/OutputLog/Private/SOutputLog.cpp
@@ -1235,7 +1235,6 @@ SOutputLog::~SOutputLog()
 	{
 		GLog->RemoveOutputDevice(this);
 	}
-	FCoreDelegates::OnHandleSystemError.RemoveAll(this);
 
 #if WITH_EDITOR
 	if (UObjectInitialized() && !GExitPurge)
