OutputLogModule (clear on pie);

Fixing the clear on pie checkbox to work by putting in the correct path.

-- Code --
OutputLogModule:
FOutputLogModule::ClearOnPIE() - Changed the path from UnrealEd.EditorPerProjectUserSettings to OutputLog.OutputLogSettings when getting the bool.

#tests PIE, made sure the output log cleared.

diff --git a/Engine/Source/Developer/OutputLog/Private/OutputLogModule.cpp b/Engine/Source/Developer/OutputLog/Private/OutputLogModule.cpp
--- a/Engine/Source/Developer/OutputLog/Private/OutputLogModule.cpp
+++ b/Engine/Source/Developer/OutputLog/Private/OutputLogModule.cpp
@@ -328,7 +328,7 @@ void FOutputLogModule::CloseDebugConsole()
 void FOutputLogModule::ClearOnPIE(const bool bIsSimulating)
 {
 	bool bClearOnPIEEnabled = false;
-	GConfig->GetBool(TEXT("/Script/UnrealEd.EditorPerProjectUserSettings"), TEXT("bEnableOutputLogClearOnPIE"), bClearOnPIEEnabled, GEditorPerProjectIni);
+	GConfig->GetBool(TEXT("/Script/OutputLog.OutputLogSettings"), TEXT("bEnableOutputLogClearOnPIE"), bClearOnPIEEnabled, GEditorPerProjectIni);
 
 	if (bClearOnPIEEnabled)
 	{
