Add ability for different verbosities in OutputLog to ignore Category Filter and always display

diff --git a/Engine/Source/Developer/OutputLog/Private/OutputLogModule.cpp b/Engine/Source/Developer/OutputLog/Private/OutputLogModule.cpp
--- a/Engine/Source/Developer/OutputLog/Private/OutputLogModule.cpp
+++ b/Engine/Source/Developer/OutputLog/Private/OutputLogModule.cpp
@@ -98,7 +98,7 @@ private:
 TSharedRef<SDockTab> FOutputLogModule::SpawnOutputLogTab(const FSpawnTabArgs& Args)
 {
 	TSharedRef<SOutputLog> NewLog = SNew(SOutputLog, false).Messages(OutputLogHistory->GetMessages());
-	NewLog->UpdateOutputLogFilter(OutputLogFilterCache->GetSelectedLogCategories(), OutputLogFilterCache->bShowErrors, OutputLogFilterCache->bShowWarnings, OutputLogFilterCache->bShowLogs);
+	NewLog->UpdateOutputLogFilter(*OutputLogFilterCache);
 
 	OutputLog = NewLog;
 
@@ -223,7 +223,7 @@ TSharedRef<SWidget> FOutputLogModule::MakeOutputLogDrawerWidget(const FSimpleDel
 			.OnCloseConsole(OnCloseConsole)
 			.Messages(OutputLogHistory->GetMessages());
 
-		OutputLogDrawerPinned->UpdateOutputLogFilter(OutputLogFilterCache->GetSelectedLogCategories(), OutputLogFilterCache->bShowErrors, OutputLogFilterCache->bShowWarnings, OutputLogFilterCache->bShowLogs);
+		OutputLogDrawerPinned->UpdateOutputLogFilter(*OutputLogFilterCache);
 
 		OutputLogDrawer = OutputLogDrawerPinned;
 	}
@@ -403,29 +403,33 @@ void FOutputLogModule::FocusOutputLog()
 
 void FOutputLogModule::UpdateOutputLogFilter(const TArray<FName>& CategoriesToShow, TOptional<bool> bShowErrors, TOptional<bool> bShowWarnings, TOptional<bool> bShowLogs)
 {
-	if (TSharedPtr<SOutputLog> SharedOutputLog = OutputLog.Pin())
-	{
-		SharedOutputLog->UpdateOutputLogFilter(CategoriesToShow, bShowErrors, bShowWarnings, bShowLogs);
-	}
+	FOutputFilterParams Params;
+	Params.bShowErrors = bShowErrors;
+	Params.bShowWarnings = bShowWarnings;
+	Params.bShowLogs = bShowLogs;
 
-	if (TSharedPtr<SOutputLog> SharedOutputLogDrawer = OutputLogDrawer.Pin())
+	UpdateOutputLogFilter(CategoriesToShow, Params);
+}
+
+void FOutputLogModule::UpdateOutputLogFilter(const TArray<FName>& CategoriesToShow, const FOutputFilterParams& InParams)
+{
+	if (InParams.bShowErrors.IsSet())
 	{
-		SharedOutputLogDrawer->UpdateOutputLogFilter(CategoriesToShow, bShowErrors, bShowWarnings, bShowLogs);
+		OutputLogFilterCache->bShowErrors = InParams.bShowErrors.GetValue();
 	}
-	
 	// Update the filter cache to these new settings
 	// This will be useful for the case where the OutputLog Drawer or Tab get created after this call
-	if (bShowErrors.IsSet())
+	if (InParams.bShowWarnings.IsSet())
 	{
-		OutputLogFilterCache->bShowErrors = bShowErrors.GetValue();
+		OutputLogFilterCache->bShowWarnings = InParams.bShowWarnings.GetValue();
 	}
-	if (bShowWarnings.IsSet())
+	if (InParams.bShowLogs.IsSet())
 	{
-		OutputLogFilterCache->bShowWarnings = bShowWarnings.GetValue();
+		OutputLogFilterCache->bShowLogs = InParams.bShowLogs.GetValue();
 	}
-	if (bShowLogs.IsSet())
+	if (InParams.IgnoreFilterVerbosities.IsSet())
 	{
-		OutputLogFilterCache->bShowLogs = bShowLogs.GetValue();
+		OutputLogFilterCache->IgnoreFilterVerbosities = InParams.IgnoreFilterVerbosities.GetValue();
 	}
 
 	OutputLogFilterCache->ClearSelectedLogCategories();
@@ -433,6 +437,16 @@ void FOutputLogModule::UpdateOutputLogFilter(const TArray<FName>& CategoriesToSh
 	{
 		OutputLogFilterCache->ToggleLogCategory(CategoryToShow);
 	}
+
+	if (TSharedPtr<SOutputLog> SharedOutputLog = OutputLog.Pin())
+	{
+		SharedOutputLog->UpdateOutputLogFilter(*OutputLogFilterCache);
+	}
+
+	if (TSharedPtr<SOutputLog> SharedOutputLogDrawer = OutputLogDrawer.Pin())
+	{
+		SharedOutputLogDrawer->UpdateOutputLogFilter(*OutputLogFilterCache);
+	}
 }
 
 void FOutputLogModule::OpenOutputLog() const
diff --git a/Engine/Source/Developer/OutputLog/Private/SOutputLog.cpp b/Engine/Source/Developer/OutputLog/Private/SOutputLog.cpp
--- a/Engine/Source/Developer/OutputLog/Private/SOutputLog.cpp
+++ b/Engine/Source/Developer/OutputLog/Private/SOutputLog.cpp
@@ -1644,70 +1644,71 @@ TSharedRef<SWidget> SOutputLog::MakeAddFilterMenu()
 {
 	FMenuBuilder MenuBuilder(/*bInShouldCloseWindowAfterMenuSelection=*/true, nullptr);
 	
+	MenuBuilder.BeginSection("OutputLogMiscEntries", LOCTEXT("OutputLogFilters", "Filters"));
+	{
+		MenuBuilder.AddMenuEntry(
+			LOCTEXT("ShowAllCategories", "Show All"),
+			LOCTEXT("ShowAllCategories_Tooltip", "Filter the Output Log to show all categories"),
+			FSlateIcon(),
+			FUIAction(FExecuteAction::CreateSP(this, &SOutputLog::CategoriesShowAll_Execute),
+				FCanExecuteAction::CreateLambda([] { return true; }),
+				FIsActionChecked::CreateSP(this, &SOutputLog::CategoriesShowAll_IsChecked)),
+			NAME_None,
+			EUserInterfaceActionType::ToggleButton
+		);
+
+		MenuBuilder.AddSubMenu(
+			LOCTEXT("Filters", "Filters"),
+			LOCTEXT("SelectCategoriesToolTip", "Set filter settings such as visible categories"),
+			FNewMenuDelegate::CreateSP(this, &SOutputLog::MakeSelectCategoriesSubMenu)
+		);
+	}
+	MenuBuilder.EndSection();
+
 	MenuBuilder.BeginSection("OutputLogVerbosityEntries", LOCTEXT("OutputLogVerbosityHeading", "Verbosity"));
 	{
 		MenuBuilder.AddMenuEntry(
 			LOCTEXT("ShowMessages", "Messages"), 
-			LOCTEXT("ShowMessages_Tooltip", "Filter the Output Log to show messages"), 
+			LOCTEXT("ShowMessages_Tooltip", "[Checked] Filter Output Log to show messages within selected categories. [-] Filter Output Log to show all messages"), 
 			FSlateIcon(), 
 			FUIAction(FExecuteAction::CreateSP(this, &SOutputLog::VerbosityLogs_Execute), 
 				FCanExecuteAction::CreateLambda([] { return true; }), 
-				FIsActionChecked::CreateSP(this, &SOutputLog::VerbosityLogs_IsChecked)), 
+				FGetActionCheckState::CreateSP(this, &SOutputLog::VerbosityLogs_IsChecked)),
 			NAME_None, 
 			EUserInterfaceActionType::ToggleButton
 		);
 
 		MenuBuilder.AddMenuEntry(
-			LOCTEXT("ShowWarnings", "Warnings"), 
-			LOCTEXT("ShowWarnings_Tooltip", "Filter the Output Log to show warnings"), 
-			FSlateIcon(), 
-			FUIAction(FExecuteAction::CreateSP(this, &SOutputLog::VerbosityWarnings_Execute), 
-				FCanExecuteAction::CreateLambda([] { return true; }), 
-				FIsActionChecked::CreateSP(this, &SOutputLog::VerbosityWarnings_IsChecked)), 
-			NAME_None, 
+			LOCTEXT("ShowWarnings", "Warnings"),
+			LOCTEXT("ShowWarnings_Tooltip", "[Checked] Filter Output Log to show warnings within selected categories. [-] Filter Output Log to show all warnings"),
+			FSlateIcon(),
+			FUIAction(FExecuteAction::CreateSP(this, &SOutputLog::VerbosityWarnings_Execute),
+				FCanExecuteAction::CreateLambda([] { return true; }),
+				FGetActionCheckState::CreateSP(this, &SOutputLog::VerbosityWarnings_IsChecked)),
+			NAME_None,
 			EUserInterfaceActionType::ToggleButton
 		);
 
 		MenuBuilder.AddMenuEntry(
-			LOCTEXT("ShowErrors", "Errors"), 
-			LOCTEXT("ShowErrors_Tooltip", "Filter the Output Log to show errors"), 
-			FSlateIcon(), 
-			FUIAction(FExecuteAction::CreateSP(this, &SOutputLog::VerbosityErrors_Execute), 
-				FCanExecuteAction::CreateLambda([] { return true; }), 
-				FIsActionChecked::CreateSP(this, &SOutputLog::VerbosityErrors_IsChecked)), 
-			NAME_None, 
+			LOCTEXT("ShowErrors", "Errors"),
+			LOCTEXT("ShowErrors_Tooltip", "[Checked] Filter Output Log to show errors within selected categories. [-] Filter Output Log to show all errors"),
+			FSlateIcon(),
+			FUIAction(FExecuteAction::CreateSP(this, &SOutputLog::VerbosityErrors_Execute),
+				FCanExecuteAction::CreateLambda([] { return true; }),
+				FGetActionCheckState::CreateSP(this, &SOutputLog::VerbosityErrors_IsChecked)),
+			NAME_None,
 			EUserInterfaceActionType::ToggleButton
 		);
 	}
 	MenuBuilder.EndSection();
 
-	MenuBuilder.BeginSection("OutputLogMiscEntries", LOCTEXT("OutputLogMiscHeading", "Miscellaneous"));
-	{
-		MenuBuilder.AddSubMenu(
-			LOCTEXT("Categories", "Categories"), 
-			LOCTEXT("SelectCategoriesToolTip", "Select Categories to display."), 
-			FNewMenuDelegate::CreateSP(this, &SOutputLog::MakeSelectCategoriesSubMenu)
-		);
-	}
-
 	return MenuBuilder.MakeWidget();
 }
 
 void SOutputLog::MakeSelectCategoriesSubMenu(FMenuBuilder& MenuBuilder)
 {
 	MenuBuilder.BeginSection("OutputLogCategoriesEntries");
-	{
-		MenuBuilder.AddMenuEntry(
-			LOCTEXT("ShowAllCategories", "Show All"),
-			LOCTEXT("ShowAllCategories_Tooltip", "Filter the Output Log to show all categories"),
-			FSlateIcon(),
-			FUIAction(FExecuteAction::CreateSP(this, &SOutputLog::CategoriesShowAll_Execute),
-			FCanExecuteAction::CreateLambda([] { return true; }),
-			FIsActionChecked::CreateSP(this, &SOutputLog::CategoriesShowAll_IsChecked)),
-			NAME_None,
-			EUserInterfaceActionType::ToggleButton
-		);
-		
+	{	
 		for (const FName& Category : Filter.GetAvailableLogCategories())
 		{
 			MenuBuilder.AddMenuEntry(
@@ -1725,24 +1726,70 @@ void SOutputLog::MakeSelectCategoriesSubMenu(FMenuBuilder& MenuBuilder)
 	MenuBuilder.EndSection();
 }
 
-bool SOutputLog::VerbosityLogs_IsChecked() const
+ECheckBoxState SOutputLog::VerbosityLogs_IsChecked() const
 {
-	return Filter.bShowLogs;
+	if (Filter.bShowLogs && Filter.IgnoreFilterVerbosities.Contains(ELogVerbosity::Log))
+	{
+		return ECheckBoxState::Undetermined;
+	}
+
+	if (Filter.bShowLogs)
+	{
+		return ECheckBoxState::Checked;
+	}
+
+	return ECheckBoxState::Unchecked;
 }
 
-bool SOutputLog::VerbosityWarnings_IsChecked() const
+ECheckBoxState SOutputLog::VerbosityWarnings_IsChecked() const
 {
-	return Filter.bShowWarnings;
+	if (Filter.bShowWarnings && Filter.IgnoreFilterVerbosities.Contains(ELogVerbosity::Warning))
+	{
+		return ECheckBoxState::Undetermined;
+	}
+
+	if (Filter.bShowWarnings)
+	{
+		return ECheckBoxState::Checked;
+	}
+
+	return ECheckBoxState::Unchecked;
 }
 
-bool SOutputLog::VerbosityErrors_IsChecked() const
+ECheckBoxState SOutputLog::VerbosityErrors_IsChecked() const
 {
-	return Filter.bShowErrors;
+	if (Filter.bShowErrors && Filter.IgnoreFilterVerbosities.Contains(ELogVerbosity::Error))
+	{
+		return ECheckBoxState::Undetermined;
+	}
+
+	if (Filter.bShowErrors)
+	{
+		return ECheckBoxState::Checked;
+	}
+
+	return ECheckBoxState::Unchecked;
 }
 
 void SOutputLog::VerbosityLogs_Execute()
 { 
-	Filter.bShowLogs = !Filter.bShowLogs;
+	// Rotate through: showing the verbosity, showing the verbosity while ignoring filter categories, and hiding the verbosity
+	if (Filter.bShowLogs)
+	{
+		if (Filter.IgnoreFilterVerbosities.Contains(ELogVerbosity::Log))
+		{
+			Filter.bShowLogs = false;
+			Filter.IgnoreFilterVerbosities.Remove(ELogVerbosity::Log);
+		}
+		else
+		{
+			Filter.IgnoreFilterVerbosities.Emplace(ELogVerbosity::Log);
+		}
+	}
+	else
+	{
+		Filter.bShowLogs = true;
+	}
 
 	// Flag the messages count as dirty
 	MessagesTextMarshaller->MarkMessagesCacheAsDirty();
@@ -1752,7 +1799,23 @@ void SOutputLog::VerbosityLogs_Execute()
 
 void SOutputLog::VerbosityWarnings_Execute()
 {
-	Filter.bShowWarnings = !Filter.bShowWarnings;
+	// Rotate through: showing the verbosity, showing the verbosity while ignoring filter categories, and hiding the verbosity
+	if (Filter.bShowWarnings)
+	{
+		if (Filter.IgnoreFilterVerbosities.Contains(ELogVerbosity::Warning))
+		{
+			Filter.bShowWarnings = false;
+			Filter.IgnoreFilterVerbosities.Remove(ELogVerbosity::Warning);
+		}
+		else
+		{
+			Filter.IgnoreFilterVerbosities.Emplace(ELogVerbosity::Warning);
+		}
+	}
+	else
+	{
+		Filter.bShowWarnings = true;
+	}
 
 	// Flag the messages count as dirty
 	MessagesTextMarshaller->MarkMessagesCacheAsDirty();
@@ -1762,7 +1825,23 @@ void SOutputLog::VerbosityWarnings_Execute()
 
 void SOutputLog::VerbosityErrors_Execute()
 {
-	Filter.bShowErrors = !Filter.bShowErrors;
+	// Rotate through: showing the verbosity, showing the verbosity while ignoring filter categories, and hiding the verbosity
+	if (Filter.bShowErrors)
+	{
+		if (Filter.IgnoreFilterVerbosities.Contains(ELogVerbosity::Error))
+		{
+			Filter.bShowErrors = false;
+			Filter.IgnoreFilterVerbosities.Remove(ELogVerbosity::Error);
+		}
+		else
+		{
+			Filter.IgnoreFilterVerbosities.Emplace(ELogVerbosity::Error);
+		}
+	}
+	else
+	{
+		Filter.bShowErrors = true;
+	}
 
 	// Flag the messages count as dirty
 	MessagesTextMarshaller->MarkMessagesCacheAsDirty();
@@ -1784,15 +1863,6 @@ void SOutputLog::CategoriesShowAll_Execute()
 {
 	Filter.bShowAllCategories = !Filter.bShowAllCategories;
 
-	Filter.ClearSelectedLogCategories();
-	if (Filter.bShowAllCategories)
-	{
-		for (const auto& AvailableCategory : Filter.GetAvailableLogCategories())
-		{
-			Filter.ToggleLogCategory(AvailableCategory);
-		}
-	}
-
 	// Flag the messages count as dirty
 	MessagesTextMarshaller->MarkMessagesCacheAsDirty();
 
@@ -1838,6 +1908,14 @@ void SOutputLog::UpdateOutputLogFilter(const TArray<FName>& CategoriesToShow, TO
 	Refresh();
 }
 
+void SOutputLog::UpdateOutputLogFilter(const FOutputLogFilter& InFilter)
+{
+	Filter = InFilter;
+	Filter.bShowAllCategories = Filter.GetSelectedLogCategories().Num() == 0 ;
+	MessagesTextMarshaller->MarkMessagesCacheAsDirty();
+	Refresh();
+}
+
 TSharedRef<SWidget> SOutputLog::GetViewButtonContent(EOutputLogSettingsMenuFlags Flags)
 {
 	TSharedPtr<FExtender> Extender;
@@ -2031,13 +2109,11 @@ bool FOutputLogFilter::IsMessageAllowed(const TSharedPtr<FOutputLogMessage>& Mes
 		{
 			return false;
 		}
-
-		if (Message->Verbosity == ELogVerbosity::Warning && !bShowWarnings)
+		else if (Message->Verbosity == ELogVerbosity::Warning && !bShowWarnings)
 		{
 			return false;
 		}
-
-		if (Message->Verbosity != ELogVerbosity::Error && Message->Verbosity != ELogVerbosity::Warning && !bShowLogs)
+		else if (Message->Verbosity != ELogVerbosity::Error && Message->Verbosity != ELogVerbosity::Warning && !bShowLogs)
 		{
 			return false;
 		}
@@ -2045,7 +2121,7 @@ bool FOutputLogFilter::IsMessageAllowed(const TSharedPtr<FOutputLogMessage>& Mes
 
 	// Filter by Category
 	{
-		if (!IsLogCategoryEnabled(Message->Category))
+		if (!bShowAllCategories && !IgnoreFilterVerbosities.Contains(Message->Verbosity) && !IsLogCategoryEnabled(Message->Category))
 		{
 			return false;
 		}
diff --git a/Engine/Source/Developer/OutputLog/Private/SOutputLog.h b/Engine/Source/Developer/OutputLog/Private/SOutputLog.h
--- a/Engine/Source/Developer/OutputLog/Private/SOutputLog.h
+++ b/Engine/Source/Developer/OutputLog/Private/SOutputLog.h
@@ -235,7 +235,7 @@ private:
 };
 
 /**
-* Holds information about filters
+* Holds information about filters_
 */
 struct FOutputLogFilter
 {
@@ -251,6 +251,9 @@ struct FOutputLogFilter
 	/** true to allow all Log Categories */
 	bool bShowAllCategories;
 
+	/** Set of Verbosity levels that will show all regardless of category filter */
+	TSet<ELogVerbosity::Type> IgnoreFilterVerbosities;
+
 	/** Enable all filters by default */
 	FOutputLogFilter() : TextFilterExpressionEvaluator(ETextFilterExpressionEvaluatorMode::BasicString)
 	{
@@ -377,6 +380,7 @@ public:
 
 	/** Change the output log's filter. If CategoriesToShow is empty, all categories will be shown. */
 	void UpdateOutputLogFilter(const TArray<FName>& CategoriesToShow, TOptional<bool> bShowErrors = TOptional<bool>(), TOptional<bool> bShowWarnings = TOptional<bool>(), TOptional<bool> bShowLogs = TOptional<bool>());
+	void UpdateOutputLogFilter(const FOutputLogFilter& InFilter);
 
 	const FOutputLogFilter& GetOutputLogFilter() { return Filter; }
 protected:
@@ -432,19 +436,19 @@ private:
 	void VerbosityLogs_Execute();
 
 	/** Returns the state of Verbosity "Logs". */
-	bool VerbosityLogs_IsChecked() const;
+	ECheckBoxState VerbosityLogs_IsChecked() const;
 
 	/** Toggles Verbosity "Warnings" true/false. */
 	void VerbosityWarnings_Execute();
 
 	/** Returns the state of Verbosity "Warnings". */
-	bool VerbosityWarnings_IsChecked() const;
+	ECheckBoxState VerbosityWarnings_IsChecked() const;
 
 	/** Toggles Verbosity "Errors" true/false. */
 	void VerbosityErrors_Execute();
 
 	/** Returns the state of Verbosity "Errors". */
-	bool VerbosityErrors_IsChecked() const;
+	ECheckBoxState VerbosityErrors_IsChecked() const;
 
 	/** Toggles All Categories true/false. */
 	void CategoriesShowAll_Execute();
diff --git a/Engine/Source/Developer/OutputLog/Public/OutputLogModule.h b/Engine/Source/Developer/OutputLog/Public/OutputLogModule.h
--- a/Engine/Source/Developer/OutputLog/Public/OutputLogModule.h
+++ b/Engine/Source/Developer/OutputLog/Public/OutputLogModule.h
@@ -71,8 +71,18 @@ public:
 
 	const TSharedPtr<SDockTab> GetOutputLogTab() const { return OutputLogTab.Pin(); }
 
+	struct FOutputFilterParams
+	{
+		TOptional<bool> bShowErrors;
+		TOptional<bool> bShowWarnings;
+		TOptional<bool> bShowLogs;
+		TOptional<TSet<ELogVerbosity::Type>> IgnoreFilterVerbosities;
+	};
+
 	/** Change the output log's filter. If CategoriesToShow is empty, all categories will be shown. */
 	void OUTPUTLOG_API UpdateOutputLogFilter(const TArray<FName>& CategoriesToShow, TOptional<bool> bShowErrors = TOptional<bool>(), TOptional<bool> bShowWarnings = TOptional<bool>(), TOptional<bool> bShowLogs = TOptional<bool>());
+	void OUTPUTLOG_API UpdateOutputLogFilter(const TArray<FName>& CategoriesToShow, const FOutputFilterParams& InParams);
+
 	/** Opens the output log tab, or brings it to front if it's already open */
 	void OUTPUTLOG_API OpenOutputLog() const;
 
