Restore Show All performing select/unselect all functionality to SOuptutLog

diff --git a/Engine/Source/Developer/OutputLog/Private/SOutputLog.cpp b/Engine/Source/Developer/OutputLog/Private/SOutputLog.cpp
--- a/Engine/Source/Developer/OutputLog/Private/SOutputLog.cpp
+++ b/Engine/Source/Developer/OutputLog/Private/SOutputLog.cpp
@@ -1863,6 +1863,15 @@ void SOutputLog::CategoriesShowAll_Execute()
 {
 	Filter.bShowAllCategories = !Filter.bShowAllCategories;
 
+	Filter.ClearSelectedLogCategories();
+	if (Filter.bShowAllCategories)
+	{
+		for (const auto& AvailableCategory : Filter.GetAvailableLogCategories())
+		{
+			Filter.ToggleLogCategory(AvailableCategory);
+		}
+	}
+
 	// Flag the messages count as dirty
 	MessagesTextMarshaller->MarkMessagesCacheAsDirty();
 
