A couple of adjustments to console input :
* Closed the suggestion box when using the keyboard shortcut (tilde key, usually) to switch between status bar console input to output log console input
* Auto-adjusted the editable text box size so that we better see the content for long console commands

diff --git a/Engine/Source/Developer/OutputLog/Private/SOutputLog.cpp b/Engine/Source/Developer/OutputLog/Private/SOutputLog.cpp
--- a/Engine/Source/Developer/OutputLog/Private/SOutputLog.cpp
+++ b/Engine/Source/Developer/OutputLog/Private/SOutputLog.cpp
@@ -207,18 +207,24 @@ void SConsoleInputBox::Construct(const FArguments& InArgs)
 			]
 
 			+SHorizontalBox::Slot()
+			.AutoWidth()
 			[
-				SAssignNew(InputText, SMultiLineEditableTextBox)
-				.Font(FOutputLogStyle::Get().GetWidgetStyle<FTextBlockStyle>("Log.Normal").Font)
-				.HintText(this, &SConsoleInputBox::GetActiveCommandExecutorHintText)
-				.AllowMultiLine(this, &SConsoleInputBox::GetActiveCommandExecutorAllowMultiLine)
-				.OnTextCommitted(this, &SConsoleInputBox::OnTextCommitted)
-				.OnTextChanged(this, &SConsoleInputBox::OnTextChanged)
-				.OnKeyCharHandler(this, &SConsoleInputBox::OnKeyCharHandler)
-				.OnKeyDownHandler(this, &SConsoleInputBox::OnKeyDownHandler)
-				.OnIsTypedCharValid(FOnIsTypedCharValid::CreateLambda([](const TCHAR InCh) { return true; })) // allow tabs to be typed into the field
-				.ClearKeyboardFocusOnCommit(false)
-				.ModiferKeyForNewLine(EModifierKey::Shift)
+				SNew(SBox)
+				.MinDesiredWidth(300.f)
+				.MaxDesiredWidth(600.f)
+				[
+					SAssignNew(InputText, SMultiLineEditableTextBox)
+					.Font(FOutputLogStyle::Get().GetWidgetStyle<FTextBlockStyle>("Log.Normal").Font)
+					.HintText(this, &SConsoleInputBox::GetActiveCommandExecutorHintText)
+					.AllowMultiLine(this, &SConsoleInputBox::GetActiveCommandExecutorAllowMultiLine)
+					.OnTextCommitted(this, &SConsoleInputBox::OnTextCommitted)
+					.OnTextChanged(this, &SConsoleInputBox::OnTextChanged)
+					.OnKeyCharHandler(this, &SConsoleInputBox::OnKeyCharHandler)
+					.OnKeyDownHandler(this, &SConsoleInputBox::OnKeyDownHandler)
+					.OnIsTypedCharValid(FOnIsTypedCharValid::CreateLambda([](const TCHAR InCh) { return true; })) // allow tabs to be typed into the field
+					.ClearKeyboardFocusOnCommit(false)
+					.ModiferKeyForNewLine(EModifierKey::Shift)
+				]
 			]
 		]
 		.MenuContent
@@ -732,6 +738,7 @@ FReply SConsoleInputBox::OnKeyDownHandler(const FGeometry& MyGeometry, const FKe
 	// Intercept the "open console" key
 	if (ActiveCommandExecutor && (ActiveCommandExecutor->AllowHotKeyClose() && ActiveCommandExecutor->GetHotKey() == InputChord))
 	{
+		SuggestionBox->SetIsOpen(false);
 		OnCloseConsole.ExecuteIfBound();
 		return FReply::Handled();
 	}
diff --git a/Engine/Source/Editor/StatusBar/Private/StatusBarSubsystem.cpp b/Engine/Source/Editor/StatusBar/Private/StatusBarSubsystem.cpp
--- a/Engine/Source/Editor/StatusBar/Private/StatusBarSubsystem.cpp
+++ b/Engine/Source/Editor/StatusBar/Private/StatusBarSubsystem.cpp
@@ -490,7 +490,6 @@ TSharedRef<SWidget> UStatusBarSubsystem::MakeStatusBarWidget(FName StatusBarName
 			.Visibility(MakeAttributeLambda(IsConsoleInputBoxBorderVisible))
 			[
 				SNew(SBox)
-				.WidthOverride(350.f)
 				[
 					ConsoleInputBox
 				]
