Truncation warning fixes for output log

diff --git a/Engine/Source/Developer/OutputLog/OutputLog.Build.cs b/Engine/Source/Developer/OutputLog/OutputLog.Build.cs
--- a/Engine/Source/Developer/OutputLog/OutputLog.Build.cs
+++ b/Engine/Source/Developer/OutputLog/OutputLog.Build.cs
@@ -44,5 +44,7 @@ public class OutputLog : ModuleRules
 			// Required for output log drawer in editor / engine builds. 
 			PrivateDependencyModuleNames.Add("Engine");
 		}
+
+		UnsafeTypeCastWarningLevel = WarningLevel.Error;
 	}
 }
diff --git a/Engine/Source/Developer/OutputLog/Private/OutputLogStyle.cpp b/Engine/Source/Developer/OutputLog/Private/OutputLogStyle.cpp
--- a/Engine/Source/Developer/OutputLog/Private/OutputLogStyle.cpp
+++ b/Engine/Source/Developer/OutputLog/Private/OutputLogStyle.cpp
@@ -46,7 +46,7 @@ FOutputLogStyle::FOutputLogStyle()
 	{
 		const UOutputLogSettings* Settings = GetDefault<UOutputLogSettings>();
 		
-		const int32 LogFontSize = Settings ? Settings->LogFontSize : 9;
+		const float LogFontSize = Settings ? static_cast<float>(Settings->LogFontSize) : 9.0f;
 
 		FTextBlockStyle NormalText = FAppStyle::GetWidgetStyle<FTextBlockStyle>("NormalText");
 
