Editor | Make it possible to listen when the Status Bar Output Log is cleared

The use case for this is implementing a plugin that displays warning and error counts on the status bar - we would like to reset the counters when the Output Log is cleared.

diff --git a/Engine/Source/Developer/OutputLog/Private/OutputLogModule.cpp b/Engine/Source/Developer/OutputLog/Private/OutputLogModule.cpp
--- a/Engine/Source/Developer/OutputLog/Private/OutputLogModule.cpp
+++ b/Engine/Source/Developer/OutputLog/Private/OutputLogModule.cpp
@@ -188,11 +188,16 @@ void FOutputLogModule::ShutdownModule()
 	OutputLogHistory.Reset();
 }
 
+static const FName ModuleName = UE_MODULE_NAME;
+
 FOutputLogModule& FOutputLogModule::Get()
 {
-	static const FName OutputLog("OutputLog");
+	return FModuleManager::Get().LoadModuleChecked<FOutputLogModule>(ModuleName);
+}
 
-	return FModuleManager::Get().LoadModuleChecked<FOutputLogModule>(OutputLog);
+FOutputLogModule* FOutputLogModule::TryGet()
+{
+	return FModuleManager::Get().GetModulePtr<FOutputLogModule>(ModuleName);
 }
 
 bool FOutputLogModule::ShouldHideConsole() const
@@ -221,6 +226,13 @@ TSharedRef<SWidget> FOutputLogModule::MakeOutputLogDrawerWidget(const FSimpleDel
 		OutputLogDrawerPinned = 
 			SNew(SOutputLog, true)
 			.OnCloseConsole(OnCloseConsole)
+			.OnClearLog_Lambda([]
+			{
+				if (FOutputLogModule* OutputLogModule = FOutputLogModule::TryGet())
+				{
+					OutputLogModule->OnOutputLogDrawerCleared.Broadcast();
+				}
+			})
 			.Messages(OutputLogHistory->GetMessages());
 
 		OutputLogDrawerPinned->UpdateOutputLogFilter(*OutputLogFilterCache);
diff --git a/Engine/Source/Developer/OutputLog/Private/SOutputLog.cpp b/Engine/Source/Developer/OutputLog/Private/SOutputLog.cpp
--- a/Engine/Source/Developer/OutputLog/Private/SOutputLog.cpp
+++ b/Engine/Source/Developer/OutputLog/Private/SOutputLog.cpp
@@ -1233,6 +1233,8 @@ void SOutputLog::Construct( const FArguments& InArgs, bool bCreateDrawerDockButt
 
 	bIsUserScrolled = false;
 	RequestForceScroll();
+
+	OnClearLogDelegate = InArgs._OnClearLog;
 }
 END_SLATE_FUNCTION_BUILD_OPTIMIZATION
 
@@ -1425,6 +1427,8 @@ void SOutputLog::OnClearLog()
 	MessagesTextMarshaller->ClearMessages();
 	MessagesTextBox->Refresh();
 	bIsUserScrolled = false;
+
+	[[maybe_unused]] bool bOnClearLogDelegateExecuted = OnClearLogDelegate.ExecuteIfBound();
 }
 
 void SOutputLog::OnHighlightCategory(FName NewCategoryToHighlight)
diff --git a/Engine/Source/Developer/OutputLog/Private/SOutputLog.h b/Engine/Source/Developer/OutputLog/Private/SOutputLog.h
--- a/Engine/Source/Developer/OutputLog/Private/SOutputLog.h
+++ b/Engine/Source/Developer/OutputLog/Private/SOutputLog.h
@@ -318,6 +318,8 @@ public:
 
 		SLATE_EVENT(FSimpleDelegate, OnCloseConsole)
 
+		SLATE_EVENT(FSimpleDelegate, OnClearLog)
+
 		/** All messages captured before this log window has been created */
 		SLATE_ARGUMENT( TArray< TSharedPtr<FOutputLogMessage> >, Messages )
 
@@ -416,6 +418,8 @@ protected:
 	/** True if the user has scrolled the window upwards */
 	bool bIsUserScrolled;
 
+	FSimpleDelegate OnClearLogDelegate; 
+
 private:
 
 	void BuildInitialLogCategoryFilter(const FArguments& InArgs);
diff --git a/Engine/Source/Developer/OutputLog/Public/OutputLogModule.h b/Engine/Source/Developer/OutputLog/Public/OutputLogModule.h
--- a/Engine/Source/Developer/OutputLog/Public/OutputLogModule.h
+++ b/Engine/Source/Developer/OutputLog/Public/OutputLogModule.h
@@ -42,6 +42,7 @@ public:
 	virtual void ShutdownModule();
 
 	static OUTPUTLOG_API FOutputLogModule& Get();
+	static OUTPUTLOG_API FOutputLogModule* TryGet();
 
 	/** Returns whether debug console widgets should be hidden */
 	virtual bool ShouldHideConsole() const;
@@ -90,6 +91,8 @@ public:
 	through OutputLogModule to avoid dependencies. */
 	virtual bool ShouldCycleToOutputLogDrawer() const;
 
+	TMulticastDelegate<void()> OnOutputLogDrawerCleared;
+
 private:
 	TSharedRef<SDockTab> SpawnOutputLogTab(const FSpawnTabArgs& Args);
 	void OnOutputLogTabClosed(TSharedRef<SDockTab> Tab);
