PR #11324: Fix the "Clear on PIE" option no longer appearing

diff --git a/Engine/Source/Developer/OutputLog/Private/SOutputLog.h b/Engine/Source/Developer/OutputLog/Private/SOutputLog.h
--- a/Engine/Source/Developer/OutputLog/Private/SOutputLog.h
+++ b/Engine/Source/Developer/OutputLog/Private/SOutputLog.h
@@ -312,9 +312,7 @@ class SOutputLog
 
 public:
 
-	SLATE_BEGIN_ARGS( SOutputLog )
-		: _Messages()
-		{}
+	SLATE_BEGIN_ARGS( SOutputLog ) {}
 
 		SLATE_EVENT(FSimpleDelegate, OnCloseConsole)
 
@@ -322,7 +320,7 @@ public:
 		SLATE_ARGUMENT( TArray< TSharedPtr<FOutputLogMessage> >, Messages )
 
 		/**  */
-		SLATE_ARGUMENT( EOutputLogSettingsMenuFlags, SettingsMenuFlags)
+		SLATE_ARGUMENT_DEFAULT( EOutputLogSettingsMenuFlags, SettingsMenuFlags ) = EOutputLogSettingsMenuFlags::None;
 
 		SLATE_ARGUMENT( FDefaultCategorySelectionMap, DefaultCategorySelection )
 
