Output Log: Refactor Settings dropdown menu to use UToolMenu to enable customization.

diff --git a/Engine/Source/Developer/OutputLog/OutputLog.Build.cs b/Engine/Source/Developer/OutputLog/OutputLog.Build.cs
--- a/Engine/Source/Developer/OutputLog/OutputLog.Build.cs
+++ b/Engine/Source/Developer/OutputLog/OutputLog.Build.cs
@@ -10,12 +10,13 @@ public class OutputLog : ModuleRules
 			new string[] {
 				"Core",
 				"CoreUObject", // @todo Mac: for some reason it's needed to link in debug on Mac
-                "InputCore",
+				"InputCore",
 				"Slate",
 				"SlateCore",
 				"TargetPlatform",
 				"DesktopPlatform",
-				"ToolWidgets", 
+				"ToolWidgets",
+				"ToolMenus",
 			}
 		);
 
diff --git a/Engine/Source/Developer/OutputLog/Private/OutputLogMenuContext.h b/Engine/Source/Developer/OutputLog/Private/OutputLogMenuContext.h
new file mode 100644
--- /dev/null
+++ b/Engine/Source/Developer/OutputLog/Private/OutputLogMenuContext.h
@@ -0,0 +1,32 @@
+// Copyright Epic Games, Inc. All Rights Reserved.
+
+#pragma once
+
+#include "Templates/SharedPointer.h"
+#include "UObject/Object.h"
+
+#include "OutputLogMenuContext.generated.h"
+
+
+class SOutputLog;
+
+
+UCLASS()
+class UOutputLogMenuContext : public UObject
+{
+	GENERATED_BODY()
+
+public:
+	void Init(const TSharedRef<SOutputLog>& InOutputLog)
+	{
+		WeakOutputLog = InOutputLog;
+	}
+
+	TSharedPtr<SOutputLog> GetOutputLog() const
+	{
+		return WeakOutputLog.Pin();
+	}
+
+private:
+	TWeakPtr<SOutputLog> WeakOutputLog;
+};
diff --git a/Engine/Source/Developer/OutputLog/Private/OutputLogModule.cpp b/Engine/Source/Developer/OutputLog/Private/OutputLogModule.cpp
--- a/Engine/Source/Developer/OutputLog/Private/OutputLogModule.cpp
+++ b/Engine/Source/Developer/OutputLog/Private/OutputLogModule.cpp
@@ -169,6 +169,8 @@ void FOutputLogModule::StartupModule()
 
 	OutputLogHistory = MakeShareable(new FOutputLogHistory);
 	OutputLogFilterCache = MakeUnique<FOutputLogFilter>();
+
+	SOutputLog::RegisterSettingsMenu();
 }
 
 void FOutputLogModule::ShutdownModule()
@@ -456,5 +458,9 @@ void FOutputLogModule::OpenOutputLog() const
 
 bool FOutputLogModule::ShouldCycleToOutputLogDrawer() const 
 {
+#if WITH_EDITOR || (IS_PROGRAM && WITH_UNREAL_DEVELOPER_TOOLS)
 	return GetDefault<UOutputLogSettings>()->bCycleToOutputLogDrawer; 
+#else
+	return true;
+#endif
 }
\ No newline at end of file
diff --git a/Engine/Source/Developer/OutputLog/Private/SOutputLog.cpp b/Engine/Source/Developer/OutputLog/Private/SOutputLog.cpp
--- a/Engine/Source/Developer/OutputLog/Private/SOutputLog.cpp
+++ b/Engine/Source/Developer/OutputLog/Private/SOutputLog.cpp
@@ -32,6 +32,8 @@
 #include "Widgets/Text/SlateEditableTextTypes.h"
 #include "OutputLogSettings.h"
 #include "OutputLogStyle.h"
+#include "OutputLogMenuContext.h"
+#include "ToolMenus.h"
 
 
 #define LOCTEXT_NAMESPACE "SOutputLog"
@@ -1099,6 +1101,22 @@ FOutputLogTextLayoutMarshaller::FOutputLogTextLayoutMarshaller(TArray< TSharedPt
 {
 }
 
+//////////////////////////////////////////////////////////////////////////
+
+namespace
+{
+	const FName SettingsMenuName("OutputLog.SettingsMenu");
+
+	const FName SettingsWordWrapEntryName("WordWrapEnable");
+	const FName SettingsTimestampsSubMenuName("TimestampsSubMenu");
+	const FName SettingsClearOnPIEEntryName("ClearOnPIE");
+
+	const FName SettingsSeparatorName("Separator");
+
+	const FName SettingsBrowseLogDirectoryEntryName("BrowseLogDirectory");
+	const FName SettingsOpenLogExternalEntryName("OpenLogExternal");
+}
+
 BEGIN_SLATE_FUNCTION_BUILD_OPTIMIZATION
 void SOutputLog::Construct( const FArguments& InArgs, bool bCreateDrawerDockButton)
 {
@@ -1116,6 +1134,9 @@ void SOutputLog::Construct( const FArguments& InArgs, bool bCreateDrawerDockButt
 		.OnVScrollBarUserScrolled(this, &SOutputLog::OnUserScrolled)
 		.ContextMenuExtender(this, &SOutputLog::ExtendTextBoxMenu);
 
+	// We take the settings bit flags passed in, and register a corresponding runtime tool menu profile.
+	const FName SettingsMenuProfileName = GetSettingsMenuProfileForFlags(InArgs._SettingsMenuFlags);
+
 	ChildSlot
 	.Padding(3)
 	[
@@ -1178,7 +1199,7 @@ void SOutputLog::Construct( const FArguments& InArgs, bool bCreateDrawerDockButt
 			[
 				SNew(SComboButton)
 				.ComboButtonStyle(FOutputLogStyle::Get(), "SimpleComboButton")
-				.OnGetMenuContent(this, &SOutputLog::GetViewButtonContent, InArgs._SettingsMenuFlags)
+				.OnGetMenuContent(this, &SOutputLog::GetSettingsMenuContent, SettingsMenuProfileName)
 				.ButtonContent()
 				[
 					SNew(SHorizontalBox)
@@ -1932,105 +1953,254 @@ void SOutputLog::UpdateOutputLogFilter(const FOutputLogFilter& InFilter)
 	Refresh();
 }
 
-TSharedRef<SWidget> SOutputLog::GetViewButtonContent(EOutputLogSettingsMenuFlags Flags)
+namespace
 {
-	TSharedPtr<FExtender> Extender;
-	FMenuBuilder MenuBuilder(true, nullptr, Extender, true);
+	template <typename TContext>
+	TSharedPtr<SOutputLog> GetWidgetFromContext(const TContext& InContext)
+	{
+		UOutputLogMenuContext* Context = InContext.template FindContext<UOutputLogMenuContext>();
+		if (!ensure(Context))
+		{
+			return nullptr;
+		}
 
-	const bool bSupportWordWrapping = (Flags & EOutputLogSettingsMenuFlags::SkipEnableWordWrapping) == EOutputLogSettingsMenuFlags::None;
-	if (bSupportWordWrapping)
+		TSharedPtr<SOutputLog> Widget = Context->GetOutputLog();
+		ensure(Widget);
+		return Widget;
+	};
+};
+
+// static
+void SOutputLog::RegisterSettingsMenu()
+{
+	// We declare the menu structure during module load, but instantiate the
+	// widget much later. Because of this, predicates/actions need to "late
+	// bind" to the instance, by pulling it back out of the FToolMenuContext
+	// or FToolMenuSection. See GetWidgetFromContext() above.
+
+	UToolMenus* ToolMenus = UToolMenus::Get();
+	if (!ensure(ToolMenus))
 	{
-		MenuBuilder.AddMenuEntry(
-			LOCTEXT("WordWrapEnabledOption", "Enable Word Wrapping"),
-			LOCTEXT("WordWrapEnabledOptionToolTip", "Enable word wrapping in the Output Log."),
-			FSlateIcon(),
-			FUIAction(
-				FExecuteAction::CreateLambda([this] {
-					// This is a toggle, hence that it is inverted
-					SetWordWrapEnabled(IsWordWrapEnabled() ? ECheckBoxState::Unchecked : ECheckBoxState::Checked);
-				}),
-				FCanExecuteAction::CreateLambda([] { return true; }),
-				FIsActionChecked::CreateSP(this, &SOutputLog::IsWordWrapEnabled)
-			),
-			NAME_None,
-			EUserInterfaceActionType::ToggleButton
-		);
+		return;
 	}
-	FText Tooltip;
 
-	#if WITH_EDITORONLY_DATA
-		Tooltip = UOutputLogSettings::StaticClass()->FindPropertyByName(
-			GET_MEMBER_NAME_CHECKED(UOutputLogSettings, LogTimestampMode))->GetToolTipText();
-	#endif // WITH_EDITORONLY_DATA
-	MenuBuilder.AddSubMenu(
-		TAttribute<FText>::CreateLambda([this]()
+	if (ensure(!ToolMenus->IsMenuRegistered(SettingsMenuName)))
+	{
+		UToolMenu* Menu = ToolMenus->RegisterMenu(SettingsMenuName);
+
+		FToolMenuSection& Section = Menu->AddSection(NAME_None);
+
+		RegisterSettingsMenu_WordWrap(Section);
+		RegisterSettingsMenu_TimestampMode(Section);
+		RegisterSettingsMenu_ClearOnPIE(Section);
+
+		Section.AddSeparator(SettingsSeparatorName);
+
+		RegisterSettingsMenu_BrowseLogs(Section);
+		RegisterSettingsMenu_OpenLogExternal(Section);
+	}
+}
+
+// static
+void SOutputLog::RegisterSettingsMenu_WordWrap(FToolMenuSection& InSection)
+{
+	FToolUIAction WordWrapAction;
+	WordWrapAction.ExecuteAction = FToolMenuExecuteAction::CreateLambda([](const FToolMenuContext& InContext)
+		{
+			if (TSharedPtr<SOutputLog> This = GetWidgetFromContext(InContext); ensure(This))
+			{
+				This->SetWordWrapEnabled(This->IsWordWrapEnabled() ? ECheckBoxState::Unchecked : ECheckBoxState::Checked);
+			}
+		});
+	WordWrapAction.GetActionCheckState = FToolMenuGetActionCheckState::CreateLambda([](const FToolMenuContext& InContext) -> ECheckBoxState
+		{
+			if (TSharedPtr<SOutputLog> This = GetWidgetFromContext(InContext); ensure(This))
 			{
-				const UEnum* Enum = StaticEnum<ELogTimes::Type>();
-				return FText::Format(LOCTEXT("TimestampsSubmenu", "Timestamp Mode: {0}"),
-					Enum->GetDisplayNameTextByIndex(this->GetSelectedTimestampMode()));
-			}),
-		Tooltip,
-		FNewMenuDelegate::CreateSP(this, &SOutputLog::AddTimestampMenuSection));
+				return This->IsWordWrapEnabled() ? ECheckBoxState::Checked : ECheckBoxState::Unchecked;
+			}
+
+			return ECheckBoxState::Unchecked;
+		});
+
+	InSection.AddMenuEntry(
+		SettingsWordWrapEntryName,
+		LOCTEXT("WordWrapEnabledOption", "Enable Word Wrapping"),
+		LOCTEXT("WordWrapEnabledOptionToolTip", "Enable word wrapping in the Output Log."),
+		FSlateIcon(),
+		WordWrapAction,
+		EUserInterfaceActionType::ToggleButton
+	);
+}
+
+// static
+void SOutputLog::RegisterSettingsMenu_TimestampMode(FToolMenuSection& InSection)
+{
+	InSection.AddDynamicEntry(SettingsTimestampsSubMenuName, FNewToolMenuSectionDelegate::CreateLambda([](FToolMenuSection& InSection)
+		{
+			FText TimestampModeTooltip;
+#if WITH_EDITORONLY_DATA
+			TimestampModeTooltip = UOutputLogSettings::StaticClass()->FindPropertyByName(
+				GET_MEMBER_NAME_CHECKED(UOutputLogSettings, LogTimestampMode))->GetToolTipText();
+#endif // WITH_EDITORONLY_DATA
 
+			if (TSharedPtr<SOutputLog> This = GetWidgetFromContext(InSection); ensure(This))
+			{
+				InSection.AddSubMenu(
+					SettingsTimestampsSubMenuName,
+					TAttribute<FText>::CreateLambda([This]()
+						{
+							const UEnum* Enum = StaticEnum<ELogTimes::Type>();
+							return FText::Format(LOCTEXT("TimestampsSubmenu", "Timestamp Mode: {0}"),
+								Enum->GetDisplayNameTextByIndex(This->GetSelectedTimestampMode()));
+						}),
+					TimestampModeTooltip,
+					FNewMenuDelegate::CreateSP(This.ToSharedRef(), &SOutputLog::AddTimestampMenuSection)
+				);
+			}
+		}));
+}
 
+// static
+void SOutputLog::RegisterSettingsMenu_ClearOnPIE(FToolMenuSection& InSection)
+{
 #if WITH_EDITOR
-	const bool bSupportClearOnPie = (Flags & EOutputLogSettingsMenuFlags::SkipClearOnPie) == EOutputLogSettingsMenuFlags::None;
-	if (bSupportClearOnPie)
+	FToolUIAction ClearOnPIEAction;
+	ClearOnPIEAction.ExecuteAction = FToolMenuExecuteAction::CreateLambda([](const FToolMenuContext& InContext)
+		{
+			if (TSharedPtr<SOutputLog> This = GetWidgetFromContext(InContext); ensure(This))
+			{
+				This->SetClearOnPIE(This->IsClearOnPIEEnabled() ? ECheckBoxState::Unchecked : ECheckBoxState::Checked);
+			}
+		});
+	ClearOnPIEAction.GetActionCheckState = FToolMenuGetActionCheckState::CreateLambda([](const FToolMenuContext& InContext) -> ECheckBoxState
+		{
+			if (TSharedPtr<SOutputLog> This = GetWidgetFromContext(InContext); ensure(This))
+			{
+				return This->IsClearOnPIEEnabled() ? ECheckBoxState::Checked : ECheckBoxState::Unchecked;
+			}
+
+			return ECheckBoxState::Unchecked;
+		});
+
+	InSection.AddMenuEntry(
+		SettingsClearOnPIEEntryName,
+		LOCTEXT("ClearOnPIE", "Clear on PIE"),
+		LOCTEXT("ClearOnPIEToolTip", "Enable clearing of the Output Log on PIE startup."),
+		FSlateIcon(),
+		ClearOnPIEAction,
+		EUserInterfaceActionType::ToggleButton
+	);
+#endif
+}
+
+// static
+void SOutputLog::RegisterSettingsMenu_BrowseLogs(FToolMenuSection& InSection)
+{
+	InSection.AddMenuEntry(
+		SettingsBrowseLogDirectoryEntryName,
+		LOCTEXT("FindSourceFile", "Open Source Location"),
+		LOCTEXT("FindSourceFileTooltip", "Opens the folder containing the source of the Output Log."),
+		FSlateIcon(FOutputLogStyle::Get().GetStyleSetName(), "OutputLog.OpenSourceLocation"),
+		FToolMenuExecuteAction::CreateLambda([](const FToolMenuContext& InContext)
+			{
+				if (TSharedPtr<SOutputLog> This = GetWidgetFromContext(InContext); ensure(This))
+				{
+					This->OpenLogFileInExplorer();
+				}
+			})
+	);
+}
+
+// static
+void SOutputLog::RegisterSettingsMenu_OpenLogExternal(FToolMenuSection& InSection)
+{
+	InSection.AddMenuEntry(
+		SettingsOpenLogExternalEntryName,
+		LOCTEXT("OpenInExternalEditor", "Open In External Editor"),
+		LOCTEXT("OpenInExternalEditorTooltip", "Opens the Output Log in the default external editor."),
+		FSlateIcon(FOutputLogStyle::Get().GetStyleSetName(), "OutputLog.OpenInExternalEditor"),
+		FToolMenuExecuteAction::CreateLambda([](const FToolMenuContext& InContext)
+			{
+				if (TSharedPtr<SOutputLog> This = GetWidgetFromContext(InContext); ensure(This))
+				{
+					This->OpenLogFileInExternalEditor();
+				}
+			})
+	);
+}
+
+FName SOutputLog::GetSettingsMenuProfileForFlags(EOutputLogSettingsMenuFlags InFlags)
+{
+	UToolMenus* ToolMenus = UToolMenus::Get();
+	if (!ensure(ToolMenus) || InFlags == EOutputLogSettingsMenuFlags::None)
 	{
-		MenuBuilder.AddMenuEntry(
-			LOCTEXT("ClearOnPIE", "Clear on PIE"),
-			LOCTEXT("ClearOnPIEToolTip", "Enable clearing of the Output Log on PIE startup."),
-			FSlateIcon(),
-			FUIAction(
-				FExecuteAction::CreateLambda([this] {
-					// This is a toggle, hence that it is inverted
-					SetClearOnPIE(IsClearOnPIEEnabled() ? ECheckBoxState::Unchecked : ECheckBoxState::Checked);
-				}),
-				FCanExecuteAction::CreateLambda([] { return true; }),
-				FIsActionChecked::CreateSP(this, &SOutputLog::IsClearOnPIEEnabled)
-			),
-			NAME_None,
-			EUserInterfaceActionType::ToggleButton
-		);
+		return NAME_None;
 	}
-#else
-	constexpr bool bSupportClearOnPie = false;
-#endif
 
-	const bool bSupportSourceLocation = (Flags & EOutputLogSettingsMenuFlags::SkipOpenSourceButton) == EOutputLogSettingsMenuFlags::None;
-	const bool bSupportExternalEditor = (Flags & EOutputLogSettingsMenuFlags::SkipOpenInExternalEditorButton) == EOutputLogSettingsMenuFlags::None;
-	const bool bNeedsSeparator = (bSupportWordWrapping || bSupportClearOnPie) && (bSupportSourceLocation || bSupportExternalEditor);
-	if (bNeedsSeparator)
+	const FName MenuProfileName = *FString::Printf(TEXT("OutputLogSettings_Flags%i"), static_cast<int32>(InFlags));
+	FToolMenuProfile* FlagsProfile = ToolMenus->FindRuntimeMenuProfile(SettingsMenuName, MenuProfileName);
+	if (!FlagsProfile)
 	{
-		MenuBuilder.AddMenuSeparator();
+		FlagsProfile = ToolMenus->AddRuntimeMenuProfile(SettingsMenuName, MenuProfileName);
+
+		const bool bSupportWordWrapping = !EnumHasAnyFlags(InFlags, EOutputLogSettingsMenuFlags::SkipEnableWordWrapping);
+		const bool bSupportClearOnPie = !EnumHasAnyFlags(InFlags, EOutputLogSettingsMenuFlags::SkipClearOnPie);
+		const bool bSupportBrowseLocation = !EnumHasAnyFlags(InFlags, EOutputLogSettingsMenuFlags::SkipOpenSourceButton);
+		const bool bSupportExternalEditor = !EnumHasAnyFlags(InFlags, EOutputLogSettingsMenuFlags::SkipOpenInExternalEditorButton);
+
+		const bool bNeedsSeparator = (bSupportWordWrapping || bSupportClearOnPie) && (bSupportBrowseLocation || bSupportExternalEditor);
+
+		if (!bSupportWordWrapping)
+		{
+			FlagsProfile->AddEntry(SettingsWordWrapEntryName)->Visibility = ECustomizedToolMenuVisibility::Hidden;
+		}
+
+		if (!bSupportClearOnPie)
+		{
+			FlagsProfile->AddEntry(SettingsClearOnPIEEntryName)->Visibility = ECustomizedToolMenuVisibility::Hidden;
+		}
+
+		if (!bNeedsSeparator)
+		{
+			FlagsProfile->AddEntry(SettingsSeparatorName)->Visibility = ECustomizedToolMenuVisibility::Hidden;
+		}
+
+		if (!bSupportBrowseLocation)
+		{
+			FlagsProfile->AddEntry(SettingsBrowseLogDirectoryEntryName)->Visibility = ECustomizedToolMenuVisibility::Hidden;
+		}
+
+		if (!bSupportExternalEditor)
+		{
+			FlagsProfile->AddEntry(SettingsOpenLogExternalEntryName)->Visibility = ECustomizedToolMenuVisibility::Hidden;
+		}
 	}
 
-	if (bSupportSourceLocation)
+	return MenuProfileName;
+}
+
+TSharedRef<SWidget> SOutputLog::GetSettingsMenuContent(FName InMenuProfileName)
+{
+	UToolMenus* ToolMenus = UToolMenus::Get();
+	if (!ensure(ToolMenus))
 	{
-		MenuBuilder.AddMenuEntry(
-			LOCTEXT("FindSourceFile", "Open Source Location"),
-			LOCTEXT("FindSourceFileTooltip", "Opens the folder containing the source of the Output Log."),
-			FSlateIcon(FOutputLogStyle::Get().GetStyleSetName(), "OutputLog.OpenSourceLocation"),
-			FUIAction(
-				FExecuteAction::CreateSP(this, &SOutputLog::OpenLogFileInExplorer)
-			)
-		);
+		return SNullWidget::NullWidget;
 	}
-	
-	if (bSupportExternalEditor)
+
+	FToolMenuContext MenuContext;
+
+	UOutputLogMenuContext* OutputLogContext = NewObject<UOutputLogMenuContext>();
+	OutputLogContext->Init(SharedThis(this));
+	MenuContext.AddObject(OutputLogContext);
+
+	if (InMenuProfileName != NAME_None)
 	{
-		// Open In External Editor
-		MenuBuilder.AddMenuEntry(
-			LOCTEXT("OpenInExternalEditor", "Open In External Editor"),
-			LOCTEXT("OpenInExternalEditorTooltip", "Opens the Output Log in the default external editor."),
-			FSlateIcon(FOutputLogStyle::Get().GetStyleSetName(), "OutputLog.OpenInExternalEditor"),
-			FUIAction(
-				FExecuteAction::CreateSP(this, &SOutputLog::OpenLogFileInExternalEditor)
-			)
-		);
+		UToolMenuProfileContext* ProfileContext = NewObject<UToolMenuProfileContext>();
+		ProfileContext->ActiveProfiles.Add(InMenuProfileName);
+		MenuContext.AddObject(ProfileContext);
 	}
 
-	return MenuBuilder.MakeWidget();
+	return ToolMenus->GenerateWidget(SettingsMenuName, MenuContext);
 }
 
 TSharedRef<SWidget> SOutputLog::CreateDrawerDockButton()
diff --git a/Engine/Source/Developer/OutputLog/Private/SOutputLog.h b/Engine/Source/Developer/OutputLog/Private/SOutputLog.h
--- a/Engine/Source/Developer/OutputLog/Private/SOutputLog.h
+++ b/Engine/Source/Developer/OutputLog/Private/SOutputLog.h
@@ -20,6 +20,7 @@
 class FMenuBuilder;
 class FOutputLogTextLayoutMarshaller;
 class FTextLayout;
+struct FToolMenuSection;
 class SMenuAnchor;
 
 /**
@@ -355,6 +356,9 @@ public:
 	 */
 	static bool CreateLogMessages(const TCHAR* V, ELogVerbosity::Type Verbosity, const class FName& Category, TArray< TSharedPtr<FOutputLogMessage> >& OutMessages);
 
+	/** Registers settings dropdown tool menu for extensibility. */
+	static void RegisterSettingsMenu();
+
 	/**
 	* Called when delete all is selected
 	*/
@@ -481,9 +485,13 @@ private:
 	void SetClearOnPIE(ECheckBoxState InValue);
 #endif
 
-	FSlateColor GetViewButtonForegroundColor() const;
-
-	TSharedRef<SWidget> GetViewButtonContent(EOutputLogSettingsMenuFlags Flags);
+	static void RegisterSettingsMenu_WordWrap(FToolMenuSection& InSection);
+	static void RegisterSettingsMenu_TimestampMode(FToolMenuSection& InSection);
+	static void RegisterSettingsMenu_ClearOnPIE(FToolMenuSection& InSection);
+	static void RegisterSettingsMenu_BrowseLogs(FToolMenuSection& InSection);
+	static void RegisterSettingsMenu_OpenLogExternal(FToolMenuSection& InSection);
+	FName GetSettingsMenuProfileForFlags(EOutputLogSettingsMenuFlags InFlags);
+	TSharedRef<SWidget> GetSettingsMenuContent(FName InMenuProfileName);
 
 	TSharedRef<SWidget> CreateDrawerDockButton();
 
diff --git a/Engine/Source/Developer/OutputLog/Public/OutputLogCreationParams.h b/Engine/Source/Developer/OutputLog/Public/OutputLogCreationParams.h
--- a/Engine/Source/Developer/OutputLog/Public/OutputLogCreationParams.h
+++ b/Engine/Source/Developer/OutputLog/Public/OutputLogCreationParams.h
@@ -10,14 +10,15 @@ enum class EOutputLogSettingsMenuFlags
 
 	/** The clear on Pie button should not be created */
 	SkipClearOnPie = 0x01,
+
 	/** The Enable world wrapping button should not be created */
 	SkipEnableWordWrapping = 0x02,
 
 	/** Skip the button that opens the source folder of the output log module */
-	SkipOpenSourceButton = 0x03,
+	SkipOpenSourceButton = 0x04,
+
 	/** Skip the button which opens the log in a text editor */
-	SkipOpenInExternalEditorButton = 0x04
-	
+	SkipOpenInExternalEditorButton = 0x08	
 };
 ENUM_CLASS_FLAGS(EOutputLogSettingsMenuFlags)
 
