Fix QuickLaunch Android

diff --git a/Engine/Source/Developer/Android/AndroidTargetPlatformControls/Private/AndroidTargetDevice.inl b/Engine/Source/Developer/Android/AndroidTargetPlatformControls/Private/AndroidTargetDevice.inl
--- a/Engine/Source/Developer/Android/AndroidTargetPlatformControls/Private/AndroidTargetDevice.inl
+++ b/Engine/Source/Developer/Android/AndroidTargetPlatformControls/Private/AndroidTargetDevice.inl
@@ -54,7 +54,7 @@ inline bool FAndroidTargetDevice::PowerOff( bool Force )
 
 inline FString FAndroidTargetDevice::GetAllDevicesName() const
 {
-	return FString::Printf(TEXT("All_%s_On_%s"), *(GetTargetPlatform().IniPlatformName()), FPlatformProcess::ComputerName());
+	return FString::Printf(TEXT("All_%s_On_%s"), *(GetPlatformSettings().IniPlatformName()), FPlatformProcess::ComputerName());
 }
 
 // cancel the running application
diff --git a/Engine/Source/Developer/Android/AndroidTargetPlatformControls/Private/AndroidTargetPlatformControls.cpp b/Engine/Source/Developer/Android/AndroidTargetPlatformControls/Private/AndroidTargetPlatformControls.cpp
--- a/Engine/Source/Developer/Android/AndroidTargetPlatformControls/Private/AndroidTargetPlatformControls.cpp
+++ b/Engine/Source/Developer/Android/AndroidTargetPlatformControls/Private/AndroidTargetPlatformControls.cpp
@@ -432,7 +432,7 @@ bool FAndroidTargetPlatformControls::HandleTicker( float DeltaTime )
 			Device->SetAuthorized(DeviceInfo.bAuthorizedDevice);
 			Device->SetVersions(DeviceInfo.SDKVersion, DeviceInfo.HumanAndroidVersion);
 
-			OnDeviceDiscovered().Broadcast(Device.ToSharedRef());
+			ITargetPlatformControls::OnDeviceDiscovered().Broadcast(Device.ToSharedRef());
 		}
 	}
 
diff --git a/Engine/Source/Developer/IOS/IOSTargetPlatform/Private/IOSTargetPlatform.cpp b/Engine/Source/Developer/IOS/IOSTargetPlatform/Private/IOSTargetPlatform.cpp
--- a/Engine/Source/Developer/IOS/IOSTargetPlatform/Private/IOSTargetPlatform.cpp
+++ b/Engine/Source/Developer/IOS/IOSTargetPlatform/Private/IOSTargetPlatform.cpp
@@ -375,7 +375,7 @@ void FIOSTargetPlatform::HandlePongMessage( const FIOSLaunchDaemonPong& Message,
 		Device->SetDeviceEndpoint(Context->GetSender());
 		Device->SetIsSimulated(Message.DeviceID.Contains(TEXT("Simulator")));
 
-		OnDeviceDiscovered().Broadcast(Device.ToSharedRef());
+		ITargetPlatformControls::OnDeviceDiscovered().Broadcast(Device.ToSharedRef());
 	}
 
 	Device->LastPinged = FDateTime::UtcNow();
@@ -412,7 +412,7 @@ void FIOSTargetPlatform::HandleDeviceConnected(const FIOSLaunchDaemonPong& Messa
 			Device->SetDeviceConnectionType(Message.DeviceConnectionType);
 			Device->SetIsSimulated(Message.DeviceID.Contains(TEXT("Simulator")));
 
-			OnDeviceDiscovered().Broadcast(Device.ToSharedRef());
+			ITargetPlatformControls::OnDeviceDiscovered().Broadcast(Device.ToSharedRef());
 		}
 		else
 		{
@@ -434,7 +434,7 @@ void FIOSTargetPlatform::HandleDeviceDisconnected(const FIOSLaunchDaemonPong& Me
 	
 	if (Device.IsValid())
 	{
-		OnDeviceLost().Broadcast(Device.ToSharedRef());
+		ITargetPlatformControls::OnDeviceLost().Broadcast(Device.ToSharedRef());
 		Devices.Remove(DeviceId);
 	}
 }
diff --git a/Engine/Source/Developer/Linux/LinuxTargetPlatform/Public/LinuxTargetPlatform.h b/Engine/Source/Developer/Linux/LinuxTargetPlatform/Public/LinuxTargetPlatform.h
--- a/Engine/Source/Developer/Linux/LinuxTargetPlatform/Public/LinuxTargetPlatform.h
+++ b/Engine/Source/Developer/Linux/LinuxTargetPlatform/Public/LinuxTargetPlatform.h
@@ -116,7 +116,7 @@ public:
 			Device->SetUserCredentials(Username, Password);
 		}
 
-		this->OnDeviceDiscovered().Broadcast(Device.ToSharedRef());
+		ITargetPlatformControls::OnDeviceDiscovered().Broadcast(Device.ToSharedRef());
 		return true;
 	}
 
diff --git a/Engine/Source/Developer/OutputLog/Private/SDeviceOutputLog.cpp b/Engine/Source/Developer/OutputLog/Private/SDeviceOutputLog.cpp
--- a/Engine/Source/Developer/OutputLog/Private/SDeviceOutputLog.cpp
+++ b/Engine/Source/Developer/OutputLog/Private/SDeviceOutputLog.cpp
@@ -104,8 +104,8 @@ void SDeviceOutputLog::Construct( const FArguments& InArgs )
 	bIsUserScrolled = false;
 	RequestForceScroll();
 	
-	ITargetPlatform::OnDeviceDiscovered().AddRaw(this, &SDeviceOutputLog::HandleTargetPlatformDeviceDiscovered);
-	ITargetPlatform::OnDeviceLost().AddRaw(this, &SDeviceOutputLog::HandleTargetPlatformDeviceLost);
+	ITargetPlatformControls::OnDeviceDiscovered().AddRaw(this, &SDeviceOutputLog::HandleTargetPlatformDeviceDiscovered);
+	ITargetPlatformControls::OnDeviceLost().AddRaw(this, &SDeviceOutputLog::HandleTargetPlatformDeviceLost);
 		
 	// Get list of available devices
 	for (ITargetPlatform* Platform : GetTargetPlatformManager()->GetTargetPlatforms())
@@ -128,8 +128,8 @@ void SDeviceOutputLog::Construct( const FArguments& InArgs )
 
 SDeviceOutputLog::~SDeviceOutputLog()
 {
-	ITargetPlatform::OnDeviceDiscovered().RemoveAll(this);
-	ITargetPlatform::OnDeviceLost().RemoveAll(this);
+	ITargetPlatformControls::OnDeviceDiscovered().RemoveAll(this);
+	ITargetPlatformControls::OnDeviceLost().RemoveAll(this);
 
 	// Clearing the pointer manually to ensure that when the pointed device output object is destroyed
 	// SDeviceOutputLog is still in a valid state in case CurrentDeviceOutputPtr wanted to dereference it.
diff --git a/Engine/Source/Developer/TargetDeviceServices/Private/Services/TargetDeviceServiceManager.cpp b/Engine/Source/Developer/TargetDeviceServices/Private/Services/TargetDeviceServiceManager.cpp
--- a/Engine/Source/Developer/TargetDeviceServices/Private/Services/TargetDeviceServiceManager.cpp
+++ b/Engine/Source/Developer/TargetDeviceServices/Private/Services/TargetDeviceServiceManager.cpp
@@ -153,8 +153,8 @@ bool FTargetDeviceServiceManager::AddTargetDevice(ITargetDevicePtr InDevice)
 
 void FTargetDeviceServiceManager::InitializeTargetPlatforms()
 {
-	ITargetPlatform::OnDeviceDiscovered().AddRaw(this, &FTargetDeviceServiceManager::HandleTargetPlatformDeviceDiscovered);
-	ITargetPlatform::OnDeviceLost().AddRaw(this, &FTargetDeviceServiceManager::HandleTargetPlatformDeviceLost);
+	ITargetPlatformControls::OnDeviceDiscovered().AddRaw(this, &FTargetDeviceServiceManager::HandleTargetPlatformDeviceDiscovered);
+	ITargetPlatformControls::OnDeviceLost().AddRaw(this, &FTargetDeviceServiceManager::HandleTargetPlatformDeviceLost);
 	
 	for (ITargetPlatform* Platform : GetTargetPlatformManager()->GetTargetPlatforms())
 	{
@@ -325,8 +325,8 @@ void FTargetDeviceServiceManager::SaveSettings()
 
 void FTargetDeviceServiceManager::ShutdownTargetPlatforms()
 {
-	ITargetPlatform::OnDeviceDiscovered().RemoveAll(this);
-	ITargetPlatform::OnDeviceLost().RemoveAll(this);
+	ITargetPlatformControls::OnDeviceDiscovered().RemoveAll(this);
+	ITargetPlatformControls::OnDeviceLost().RemoveAll(this);
 }
 
 
diff --git a/Engine/Source/Developer/TargetPlatform/Public/Interfaces/ITargetPlatform.h b/Engine/Source/Developer/TargetPlatform/Public/Interfaces/ITargetPlatform.h
--- a/Engine/Source/Developer/TargetPlatform/Public/Interfaces/ITargetPlatform.h
+++ b/Engine/Source/Developer/TargetPlatform/Public/Interfaces/ITargetPlatform.h
@@ -50,12 +50,14 @@ public:
 	 * Gets an event delegate that is executed when a new target device has been discovered.
 	 */
 	DECLARE_EVENT_OneParam(ITargetPlatform, FOnTargetDeviceDiscovered, ITargetDeviceRef /*DiscoveredDevice*/);
+	UE_DEPRECATED(5.4, "ITargetPlatform::OnDeviceDiscovered is deprecated, use ITargetPlatformControls::OnDeviceDiscovered instead")
 	static TARGETPLATFORM_API FOnTargetDeviceDiscovered& OnDeviceDiscovered();
 
 	/**
 	 * Gets an event delegate that is executed when a target device has been lost, i.e. disconnected or timed out.
 	 */
 	DECLARE_EVENT_OneParam(ITargetPlatform, FOnTargetDeviceLost, ITargetDeviceRef /*LostDevice*/);
+	UE_DEPRECATED(5.4, "ITargetPlatform::OnDeviceLost is deprecated, use ITargetPlatformControls::OnDeviceLost instead")
 	static TARGETPLATFORM_API FOnTargetDeviceLost& OnDeviceLost();
 
 public:
