[trivial] set a default LogFontSize for OutputLog

 - code default size of 0 was only overridden if BaseEditorPerProjectUserSettings.ini existed which isn't the case for staged Programs

#rnx

diff --git a/Engine/Source/Developer/OutputLog/Public/OutputLogSettings.h b/Engine/Source/Developer/OutputLog/Public/OutputLogSettings.h
--- a/Engine/Source/Developer/OutputLog/Public/OutputLogSettings.h
+++ b/Engine/Source/Developer/OutputLog/Public/OutputLogSettings.h
@@ -37,6 +37,7 @@ class OUTPUTLOG_API UOutputLogSettings : public UObject
 public:
 	UOutputLogSettings()
 	{
+		LogFontSize = 9;
 		bCycleToOutputLogDrawer = true;
 		LogTimestampMode = ELogTimes::None;
 	}
