Add tooltip with cvar help text to console suggestions

diff --git a/Engine/Plugins/Developer/Concert/ConcertApp/MultiUserServer/Source/MultiUserServer/Private/ConcertConsoleCommandExecutor.cpp b/Engine/Plugins/Developer/Concert/ConcertApp/MultiUserServer/Source/MultiUserServer/Private/ConcertConsoleCommandExecutor.cpp
--- a/Engine/Plugins/Developer/Concert/ConcertApp/MultiUserServer/Source/MultiUserServer/Private/ConcertConsoleCommandExecutor.cpp
+++ b/Engine/Plugins/Developer/Concert/ConcertApp/MultiUserServer/Source/MultiUserServer/Private/ConcertConsoleCommandExecutor.cpp
@@ -35,7 +35,7 @@ namespace UE::MultiUserServer
 		return LOCTEXT("ConsoleCommandExecutorHintText", "Enter Console Command");
 	}
 
-	void FConcertConsoleCommandExecutor::GetAutoCompleteSuggestions(const TCHAR* Input, TArray<FString>& Out)
+	void FConcertConsoleCommandExecutor::GetSuggestedCompletions(const TCHAR* Input, TArray<FConsoleSuggestion>& Out)
 	{
 		auto OnConsoleVariable = [&Out](const TCHAR *Name, IConsoleObject* CVar)
 		{
@@ -44,11 +44,14 @@ namespace UE::MultiUserServer
 				return;
 			}
 
-			Out.Add(Name);
+			Out.Add(FConsoleSuggestion(Name, CVar->GetHelp()));
 		};
 
 		IConsoleManager::Get().ForEachConsoleObjectThatContains(FConsoleObjectVisitor::CreateLambda(OnConsoleVariable), Input);
-		Out.Append(GetDefault<UConsoleSettings>()->GetFilteredManualAutoCompleteCommands(Input));
+		for (const FString& Name : GetDefault<UConsoleSettings>()->GetFilteredManualAutoCompleteCommands(Input))
+		{
+			Out.Add(FConsoleSuggestion(Name, FString()));
+		}
 	}
 
 	void FConcertConsoleCommandExecutor::GetExecHistory(TArray<FString>& Out)
diff --git a/Engine/Plugins/Developer/Concert/ConcertApp/MultiUserServer/Source/MultiUserServer/Private/ConcertConsoleCommandExecutor.h b/Engine/Plugins/Developer/Concert/ConcertApp/MultiUserServer/Source/MultiUserServer/Private/ConcertConsoleCommandExecutor.h
--- a/Engine/Plugins/Developer/Concert/ConcertApp/MultiUserServer/Source/MultiUserServer/Private/ConcertConsoleCommandExecutor.h
+++ b/Engine/Plugins/Developer/Concert/ConcertApp/MultiUserServer/Source/MultiUserServer/Private/ConcertConsoleCommandExecutor.h
@@ -20,7 +20,7 @@ namespace UE::MultiUserServer
 		virtual FText GetDisplayName() const override;
 		virtual FText GetDescription() const override;
 		virtual FText GetHintText() const override;
-		virtual void GetAutoCompleteSuggestions(const TCHAR* Input, TArray<FString>& Out) override;
+		virtual void GetSuggestedCompletions(const TCHAR* Input, TArray<FConsoleSuggestion>& Out) override;
 		virtual void GetExecHistory(TArray<FString>& Out) override;
 		virtual bool Exec(const TCHAR* Input) override;
 		virtual bool AllowHotKeyClose() const override;
diff --git a/Engine/Plugins/Experimental/PythonScriptPlugin/Source/PythonScriptPlugin/Private/PythonScriptPlugin.cpp b/Engine/Plugins/Experimental/PythonScriptPlugin/Source/PythonScriptPlugin/Private/PythonScriptPlugin.cpp
--- a/Engine/Plugins/Experimental/PythonScriptPlugin/Source/PythonScriptPlugin/Private/PythonScriptPlugin.cpp
+++ b/Engine/Plugins/Experimental/PythonScriptPlugin/Source/PythonScriptPlugin/Private/PythonScriptPlugin.cpp
@@ -245,7 +245,7 @@ FText FPythonCommandExecutor::GetHintText() const
 	return LOCTEXT("PythonCommandExecutorHintText", "Enter Python script or a filename");
 }
 
-void FPythonCommandExecutor::GetAutoCompleteSuggestions(const TCHAR* Input, TArray<FString>& Out)
+void FPythonCommandExecutor::GetSuggestedCompletions(const TCHAR* Input, TArray<FConsoleSuggestion>& Out)
 {
 }
 
@@ -324,7 +324,7 @@ FText FPythonREPLCommandExecutor::GetHintText() const
 	return LOCTEXT("PythonREPLCommandExecutorHintText", "Enter a Python statement");
 }
 
-void FPythonREPLCommandExecutor::GetAutoCompleteSuggestions(const TCHAR* Input, TArray<FString>& Out)
+void FPythonREPLCommandExecutor::GetSuggestedCompletions(const TCHAR* Input, TArray<FConsoleSuggestion>& Out)
 {
 }
 
diff --git a/Engine/Plugins/Experimental/PythonScriptPlugin/Source/PythonScriptPlugin/Private/PythonScriptPlugin.h b/Engine/Plugins/Experimental/PythonScriptPlugin/Source/PythonScriptPlugin/Private/PythonScriptPlugin.h
--- a/Engine/Plugins/Experimental/PythonScriptPlugin/Source/PythonScriptPlugin/Private/PythonScriptPlugin.h
+++ b/Engine/Plugins/Experimental/PythonScriptPlugin/Source/PythonScriptPlugin/Private/PythonScriptPlugin.h
@@ -39,7 +39,7 @@ public:
 	virtual FText GetDisplayName() const override;
 	virtual FText GetDescription() const override;
 	virtual FText GetHintText() const override;
-	virtual void GetAutoCompleteSuggestions(const TCHAR* Input, TArray<FString>& Out) override;
+	virtual void GetSuggestedCompletions(const TCHAR* Input, TArray<FConsoleSuggestion>& Out) override;
 	virtual void GetExecHistory(TArray<FString>& Out) override;
 	virtual bool Exec(const TCHAR* Input) override;
 	virtual bool AllowHotKeyClose() const override;
@@ -64,7 +64,7 @@ public:
 	virtual FText GetDisplayName() const override;
 	virtual FText GetDescription() const override;
 	virtual FText GetHintText() const override;
-	virtual void GetAutoCompleteSuggestions(const TCHAR* Input, TArray<FString>& Out) override;
+	virtual void GetSuggestedCompletions(const TCHAR* Input, TArray<FConsoleSuggestion>& Out) override;
 	virtual void GetExecHistory(TArray<FString>& Out) override;
 	virtual bool Exec(const TCHAR* Input) override;
 	virtual bool AllowHotKeyClose() const override;
diff --git a/Engine/Source/Developer/OutputLog/Private/SOutputLog.cpp b/Engine/Source/Developer/OutputLog/Private/SOutputLog.cpp
--- a/Engine/Source/Developer/OutputLog/Private/SOutputLog.cpp
+++ b/Engine/Source/Developer/OutputLog/Private/SOutputLog.cpp
@@ -241,7 +241,7 @@ void SConsoleInputBox::Construct(const FArguments& InArgs)
 				.MinDesiredWidth(300.f)
 				.MaxDesiredWidth(this, &SConsoleInputBox::GetSelectionListMaxWidth)
 				[
-					SAssignNew(SuggestionListView, SListView< TSharedPtr<FString> >)
+					SAssignNew(SuggestionListView, SListView< TSharedPtr<FConsoleSuggestion> >)
 					.ListItemsSource(&Suggestions.SuggestionsList)
 					.SelectionMode( ESelectionMode::Single )							// Ideally the mouse over would not highlight while keyboard controls the UI
 					.OnGenerateRow(this, &SConsoleInputBox::MakeSuggestionListItemWidget)
@@ -269,14 +269,14 @@ void SConsoleInputBox::Tick( const FGeometry& AllottedGeometry, const double InC
 }
 
 
-void SConsoleInputBox::SuggestionSelectionChanged(TSharedPtr<FString> NewValue, ESelectInfo::Type SelectInfo)
+void SConsoleInputBox::SuggestionSelectionChanged(TSharedPtr<FConsoleSuggestion> NewValue, ESelectInfo::Type SelectInfo)
 {
 	if(bIgnoreUIUpdate)
 	{
 		return;
 	}
 
-	Suggestions.SelectedSuggestion = Suggestions.SuggestionsList.IndexOfByPredicate([&NewValue](const TSharedPtr<FString>& InSuggestion)
+	Suggestions.SelectedSuggestion = Suggestions.SuggestionsList.IndexOfByPredicate([&NewValue](const TSharedPtr<FConsoleSuggestion>& InSuggestion)
 	{
 		return InSuggestion == NewValue;
 	});
@@ -309,11 +309,11 @@ FOptionalSize SConsoleInputBox::GetSelectionListMaxWidth() const
 	return FMath::Max(300.0f, WidgetWorkArea.GetSize().X - 12.0f);
 }
 
-TSharedRef<ITableRow> SConsoleInputBox::MakeSuggestionListItemWidget(TSharedPtr<FString> Text, const TSharedRef<STableViewBase>& OwnerTable)
+TSharedRef<ITableRow> SConsoleInputBox::MakeSuggestionListItemWidget(TSharedPtr<FConsoleSuggestion> Suggestion, const TSharedRef<STableViewBase>& OwnerTable)
 {
-	check(Text.IsValid());
+	check(Suggestion.IsValid());
 
-	FString SanitizedText = *Text;
+	FString SanitizedText = Suggestion->Name;
 	SanitizedText.ReplaceInline(TEXT("\r\n"), TEXT("\n"), ESearchCase::CaseSensitive);
 	SanitizedText.ReplaceInline(TEXT("\r"), TEXT(" "), ESearchCase::CaseSensitive);
 	SanitizedText.ReplaceInline(TEXT("\n"), TEXT(" "), ESearchCase::CaseSensitive);
@@ -326,6 +326,7 @@ TSharedRef<ITableRow> SConsoleInputBox::MakeSuggestionListItemWidget(TSharedPtr<
 			.TextStyle(FOutputLogStyle::Get(), "Log.Normal")
 			.HighlightText(Suggestions.SuggestionsHighlight)
 			.ColorAndOpacity(FSlateColor::UseForeground())
+			.ToolTipText(FText::FromString(Suggestion->Help))
 		];
 }
 
@@ -339,11 +340,11 @@ void SConsoleInputBox::OnTextChanged(const FText& InText)
 	const FString& InputTextStr = InputText->GetText().ToString();
 	if(!InputTextStr.IsEmpty())
 	{
-		TArray<FString> AutoCompleteList;
+		TArray<FConsoleSuggestion> AutoCompleteList;
 		
 		if (ActiveCommandExecutor)
 		{
-			ActiveCommandExecutor->GetAutoCompleteSuggestions(*InputTextStr, AutoCompleteList);
+			ActiveCommandExecutor->GetSuggestedCompletions(*InputTextStr, AutoCompleteList);
 		}
 		else
 		{
@@ -360,31 +361,30 @@ void SConsoleInputBox::OnTextChanged(const FText& InText)
 					return;
 				}
 
-				AutoCompleteList.Add(Name);
+				AutoCompleteList.Add(FConsoleSuggestion(Name, CVar->GetHelp()));
 			};
 
 			IConsoleManager::Get().ForEachConsoleObjectThatContains(FConsoleObjectVisitor::CreateLambda(OnConsoleVariable), *InputTextStr);
-			AutoCompleteList.Append(GetDefault<UConsoleSettings>()->GetFilteredManualAutoCompleteCommands(InputTextStr));
+			//AutoCompleteList.Append(GetDefault<UConsoleSettings>()->GetFilteredManualAutoCompleteCommands(InputTextStr));
 		}
-		AutoCompleteList.Sort([InputTextStr](const FString& A, const FString& B)
+		AutoCompleteList.Sort([InputTextStr](const FConsoleSuggestion& A, const FConsoleSuggestion& B)
 		{ 
-			if (A.StartsWith(InputTextStr))
+			if (A.Name.StartsWith(InputTextStr))
 			{
-				if (!B.StartsWith(InputTextStr))
+				if (!B.Name.StartsWith(InputTextStr))
 				{
 					return true;
 				}
 			}
 			else
 			{
-				if (B.StartsWith(InputTextStr))
+				if (B.Name.StartsWith(InputTextStr))
 				{
 					return false;
 				}
 			}
 
-			return A < B;
-
+			return A.Name < B.Name;
 		});
 
 
@@ -479,14 +479,19 @@ FReply SConsoleInputBox::OnPreviewKeyDown(const FGeometry& MyGeometry, const FKe
 			const bool bShowHistory = InputText->GetText().IsEmpty() || KeyEvent.IsControlDown();
 			if (bShowHistory)
 			{
-				TArray<FString> History;
+				TArray<FString> HistoryNames;
 				if (ActiveCommandExecutor)
 				{
-					ActiveCommandExecutor->GetExecHistory(History);
+					ActiveCommandExecutor->GetExecHistory(HistoryNames);
 				}
 				else
 				{
-					IConsoleManager::Get().GetConsoleHistory(TEXT(""), History);
+					IConsoleManager::Get().GetConsoleHistory(TEXT(""), HistoryNames);
+				}
+				TArray<FConsoleSuggestion> History;
+				for (const FString& Name : HistoryNames)
+				{
+					History.Add(FConsoleSuggestion(Name, FString()));
 				}
 				SetSuggestions(History, FText::GetEmpty());
 				
@@ -531,12 +536,12 @@ FReply SConsoleInputBox::OnPreviewKeyDown(const FGeometry& MyGeometry, const FKe
 	return FReply::Unhandled();
 }
 
-void SConsoleInputBox::SetSuggestions(TArray<FString>& Elements, FText Highlight)
+void SConsoleInputBox::SetSuggestions(TArray<FConsoleSuggestion>& Elements, FText Highlight)
 {
 	FString SelectionText;
 	if (Suggestions.HasSelectedSuggestion())
 	{
-		SelectionText = *Suggestions.GetSelectedSuggestion();
+		SelectionText = Suggestions.GetSelectedSuggestion()->Name;
 	}
 
 	Suggestions.Reset();
@@ -544,9 +549,9 @@ void SConsoleInputBox::SetSuggestions(TArray<FString>& Elements, FText Highlight
 
 	for(int32 i = 0; i < Elements.Num(); ++i)
 	{
-		Suggestions.SuggestionsList.Add(MakeShared<FString>(Elements[i]));
+		Suggestions.SuggestionsList.Add(MakeShared<FConsoleSuggestion>(Elements[i]));
 
-		if (Elements[i] == SelectionText)
+		if (Elements[i].Name == SelectionText)
 		{
 			Suggestions.SelectedSuggestion = i;
 		}
@@ -582,12 +587,12 @@ void SConsoleInputBox::MarkActiveSuggestion()
 	bIgnoreUIUpdate = true;
 	if (Suggestions.HasSelectedSuggestion())
 	{
-		TSharedPtr<FString> SelectedSuggestion = Suggestions.GetSelectedSuggestion();
+		TSharedPtr<FConsoleSuggestion> SelectedSuggestion = Suggestions.GetSelectedSuggestion();
 
 		SuggestionListView->SetSelection(SelectedSuggestion);
 		SuggestionListView->RequestScrollIntoView(SelectedSuggestion);	// Ideally this would only scroll if outside of the view
 
-		InputText->SetText(FText::FromString(*SelectedSuggestion));
+		InputText->SetText(FText::FromString(SelectedSuggestion->Name));
 	}
 	else
 	{
diff --git a/Engine/Source/Developer/OutputLog/Private/SOutputLog.h b/Engine/Source/Developer/OutputLog/Private/SOutputLog.h
--- a/Engine/Source/Developer/OutputLog/Private/SOutputLog.h
+++ b/Engine/Source/Developer/OutputLog/Private/SOutputLog.h
@@ -109,11 +109,11 @@ protected:
 	FOptionalSize GetSelectionListMaxWidth() const;
 
 	/** Makes the widget for the suggestions messages in the list view */
-	TSharedRef<ITableRow> MakeSuggestionListItemWidget(TSharedPtr<FString> Message, const TSharedRef<STableViewBase>& OwnerTable);
+	TSharedRef<ITableRow> MakeSuggestionListItemWidget(TSharedPtr<FConsoleSuggestion> Message, const TSharedRef<STableViewBase>& OwnerTable);
 
-	void SuggestionSelectionChanged(TSharedPtr<FString> NewValue, ESelectInfo::Type SelectInfo);
+	void SuggestionSelectionChanged(TSharedPtr<FConsoleSuggestion> NewValue, ESelectInfo::Type SelectInfo);
 		
-	void SetSuggestions(TArray<FString>& Elements, FText Highlight);
+	void SetSuggestions(TArray<FConsoleSuggestion>& Elements, FText Highlight);
 
 	void MarkActiveSuggestion();
 
@@ -182,7 +182,7 @@ private:
 			}
 		}
 
-		TSharedPtr<FString> GetSelectedSuggestion() const
+		TSharedPtr<FConsoleSuggestion> GetSelectedSuggestion() const
 		{
 			return SuggestionsList.IsValidIndex(SelectedSuggestion) ? SuggestionsList[SelectedSuggestion] : nullptr;
 		}
@@ -191,7 +191,7 @@ private:
 		int32 SelectedSuggestion;
 
 		/** All log messages stored in this widget for the list view */
-		TArray<TSharedPtr<FString>> SuggestionsList;
+		TArray<TSharedPtr<FConsoleSuggestion>> SuggestionsList;
 
 		/** Highlight text to use for the suggestions list */
 		FText SuggestionsHighlight;
@@ -204,7 +204,7 @@ private:
 	TSharedPtr< SMenuAnchor > SuggestionBox;
 
 	/** The list view for showing all log messages. Should be replaced by a full text editor */
-	TSharedPtr< SListView< TSharedPtr<FString> > > SuggestionListView;
+	TSharedPtr< SListView< TSharedPtr<FConsoleSuggestion> > > SuggestionListView;
 
 	/** Active list of suggestions */
 	FSuggestions Suggestions;
diff --git a/Engine/Source/Editor/UnrealEd/Private/Toolkits/FConsoleCommandExecutor.cpp b/Engine/Source/Editor/UnrealEd/Private/Toolkits/FConsoleCommandExecutor.cpp
--- a/Engine/Source/Editor/UnrealEd/Private/Toolkits/FConsoleCommandExecutor.cpp
+++ b/Engine/Source/Editor/UnrealEd/Private/Toolkits/FConsoleCommandExecutor.cpp
@@ -36,7 +36,7 @@ FText FConsoleCommandExecutor::GetHintText() const
 	return LOCTEXT("ConsoleCommandExecutorHintText", "Enter Console Command");
 }
 
-void FConsoleCommandExecutor::GetAutoCompleteSuggestions(const TCHAR* Input, TArray<FString>& Out)
+void FConsoleCommandExecutor::GetSuggestedCompletions(const TCHAR* Input, TArray<FConsoleSuggestion>& Out)
 {
 	auto OnConsoleVariable = [&Out](const TCHAR *Name, IConsoleObject* CVar)
 	{
@@ -51,11 +51,14 @@ void FConsoleCommandExecutor::GetAutoCompleteSuggestions(const TCHAR* Input, TAr
 			return;
 		}
 
-		Out.Add(Name);
+		Out.Add(FConsoleSuggestion(Name, CVar->GetHelp()));
 	};
 
 	IConsoleManager::Get().ForEachConsoleObjectThatContains(FConsoleObjectVisitor::CreateLambda(OnConsoleVariable), Input);
-	Out.Append(GetDefault<UConsoleSettings>()->GetFilteredManualAutoCompleteCommands(Input));
+	for (const FString& CommandName : GetDefault<UConsoleSettings>()->GetFilteredManualAutoCompleteCommands(Input))
+	{
+		Out.Add(FConsoleSuggestion(CommandName, FString()));
+	}
 }
 
 void FConsoleCommandExecutor::GetExecHistory(TArray<FString>& Out)
diff --git a/Engine/Source/Editor/UnrealEd/Public/Toolkits/FConsoleCommandExecutor.h b/Engine/Source/Editor/UnrealEd/Public/Toolkits/FConsoleCommandExecutor.h
--- a/Engine/Source/Editor/UnrealEd/Public/Toolkits/FConsoleCommandExecutor.h
+++ b/Engine/Source/Editor/UnrealEd/Public/Toolkits/FConsoleCommandExecutor.h
@@ -18,7 +18,7 @@ public:
 	UNREALED_API virtual FText GetDisplayName() const override;
 	UNREALED_API virtual FText GetDescription() const override;
 	UNREALED_API virtual FText GetHintText() const override;
-	UNREALED_API virtual void GetAutoCompleteSuggestions(const TCHAR* Input, TArray<FString>& Out) override;
+	UNREALED_API virtual void GetSuggestedCompletions(const TCHAR* Input, TArray<FConsoleSuggestion>& Out) override;
 	UNREALED_API virtual void GetExecHistory(TArray<FString>& Out) override;
 	UNREALED_API virtual bool Exec(const TCHAR* Input) override;
 	UNREALED_API virtual bool AllowHotKeyClose() const override;
diff --git a/Engine/Source/Runtime/Core/Public/HAL/IConsoleManager.h b/Engine/Source/Runtime/Core/Public/HAL/IConsoleManager.h
--- a/Engine/Source/Runtime/Core/Public/HAL/IConsoleManager.h
+++ b/Engine/Source/Runtime/Core/Public/HAL/IConsoleManager.h
@@ -717,6 +717,19 @@ private:
 	FDelegateHandle Handle;
 };
 
+struct FConsoleSuggestion
+{
+	FString Name;
+	FString Help;
+
+	FConsoleSuggestion(FString InName, FString InHelp)
+	: Name(InName)
+	, Help(InHelp)
+	{}
+
+	FConsoleSuggestion()
+	{}
+};
 
 /**
  * Handles executing console commands
@@ -756,9 +769,15 @@ public:
 	virtual FText GetHintText() const = 0;
 
 	/**
-	 * Get the list of auto-complete suggestions for the given command.
-	 */
-	virtual void GetAutoCompleteSuggestions(const TCHAR* Input, TArray<FString>& Out) = 0;
+	* Get the list of auto-complete suggestions for the given command.
+	*/
+	virtual void GetSuggestedCompletions(const TCHAR* Input, TArray<FConsoleSuggestion>& Out) = 0;
+
+	/**
+	* Get the list of auto-complete suggestions for the given command.
+	*/
+	UE_DEPRECATED(5.5, "Use GetSuggestedCompletions instead")
+	virtual void GetAutoCompleteSuggestions(const TCHAR* Input, TArray<FString>& Out){};
 
 	/**
 	 * Get the list of commands that this executor has recently processed.
diff --git a/Engine/Source/Runtime/Engine/Private/Engine/EngineConsoleCommandExecutor.cpp b/Engine/Source/Runtime/Engine/Private/Engine/EngineConsoleCommandExecutor.cpp
--- a/Engine/Source/Runtime/Engine/Private/Engine/EngineConsoleCommandExecutor.cpp
+++ b/Engine/Source/Runtime/Engine/Private/Engine/EngineConsoleCommandExecutor.cpp
@@ -48,7 +48,7 @@ FText FEngineConsoleCommandExecutor::GetHintText() const
 	return LOCTEXT("HintText", "Enter Console Command");
 }
 
-void FEngineConsoleCommandExecutor::GetAutoCompleteSuggestions(const TCHAR* Input, TArray<FString>& Out)
+void FEngineConsoleCommandExecutor::GetSuggestedCompletions(const TCHAR* Input, TArray<FConsoleSuggestion>& Out)
 {
 	const auto OnConsoleVariable = [&Out](const TCHAR* Name, IConsoleObject* CVar)
 	{
@@ -62,11 +62,14 @@ void FEngineConsoleCommandExecutor::GetAutoCompleteSuggestions(const TCHAR* Inpu
 		{
 			return;
 		}
-		Out.Add(Name);
+		Out.Add(FConsoleSuggestion(Name, CVar->GetHelp()));
 	};
 
 	IConsoleManager::Get().ForEachConsoleObjectThatContains(FConsoleObjectVisitor::CreateLambda(OnConsoleVariable), Input);
-	Out.Append(GetDefault<UConsoleSettings>()->GetFilteredManualAutoCompleteCommands(Input));
+	for (const FString& CommandName : GetDefault<UConsoleSettings>()->GetFilteredManualAutoCompleteCommands(Input))
+	{
+		Out.Add(FConsoleSuggestion(CommandName, FString()));
+	}
 }
 
 void FEngineConsoleCommandExecutor::GetExecHistory(TArray<FString>& Out)
diff --git a/Engine/Source/Runtime/Engine/Private/Engine/EngineConsoleCommandExecutor.h b/Engine/Source/Runtime/Engine/Private/Engine/EngineConsoleCommandExecutor.h
--- a/Engine/Source/Runtime/Engine/Private/Engine/EngineConsoleCommandExecutor.h
+++ b/Engine/Source/Runtime/Engine/Private/Engine/EngineConsoleCommandExecutor.h
@@ -19,7 +19,7 @@ public:
 	FText GetDescription() const final;
 	FText GetHintText() const final;
 
-	void GetAutoCompleteSuggestions(const TCHAR* Input, TArray<FString>& Out) final;
+	void GetSuggestedCompletions(const TCHAR* Input, TArray<FConsoleSuggestion>& Out) final;
 	void GetExecHistory(TArray<FString>& Out) final;
 	bool Exec(const TCHAR* Input) final;
 
