Editor | Make it possible to listen when the Status Bar Output Log is cleared

The use case for this is implementing a plugin that displays warning and error counts on the status bar - we would like to reset the counters when the Output Log is cleared.

NOTE: This is a resubmit of a fairly old CL that got backed out due to issues with compiling on clang. I have altered the problematic code to make it compile.

Note: The PF is actually to make submit tool work but actually its 6622a95a37485b66bb46449f, submit tool is failing because its from a different stream.

diff --git a/Engine/Source/Developer/OutputLog/Private/OutputLogModule.cpp b/Engine/Source/Developer/OutputLog/Private/OutputLogModule.cpp
--- a/Engine/Source/Developer/OutputLog/Private/OutputLogModule.cpp
+++ b/Engine/Source/Developer/OutputLog/Private/OutputLogModule.cpp
@@ -190,11 +190,16 @@ void FOutputLogModule::ShutdownModule()
 	OutputLogHistory.Reset();
 }
 
+const FName FOutputLogModule::ModuleName = UE_MODULE_NAME;
+
 FOutputLogModule& FOutputLogModule::Get()
 {
-	static const FName OutputLog("OutputLog");
+	return FModuleManager::Get().LoadModuleChecked<FOutputLogModule>(ModuleName);
+}
 
-	return FModuleManager::Get().LoadModuleChecked<FOutputLogModule>(OutputLog);
+FOutputLogModule* FOutputLogModule::TryGet()
+{
+	return FModuleManager::Get().GetModulePtr<FOutputLogModule>(ModuleName);
 }
 
 bool FOutputLogModule::ShouldHideConsole() const
@@ -223,6 +228,13 @@ TSharedRef<SWidget> FOutputLogModule::MakeOutputLogDrawerWidget(const FSimpleDel
 		OutputLogDrawerPinned = 
 			SNew(SOutputLog, true)
 			.OnCloseConsole(OnCloseConsole)
+			.OnClearLog_Lambda([]
+			{
+				if (FOutputLogModule* OutputLogModule = FOutputLogModule::TryGet())
+				{
+					OutputLogModule->OnOutputLogDrawerCleared.Broadcast();
+				}
+			})
 			.Messages(OutputLogHistory->GetMessages());
 
 		OutputLogDrawerPinned->UpdateOutputLogFilter(*OutputLogFilterCache);
diff --git a/Engine/Source/Developer/OutputLog/Private/SOutputLog.cpp b/Engine/Source/Developer/OutputLog/Private/SOutputLog.cpp
--- a/Engine/Source/Developer/OutputLog/Private/SOutputLog.cpp
+++ b/Engine/Source/Developer/OutputLog/Private/SOutputLog.cpp
@@ -1261,6 +1261,8 @@ void SOutputLog::Construct( const FArguments& InArgs, bool bCreateDrawerDockButt
 
 	bIsUserScrolled = false;
 	RequestForceScroll();
+
+	OnClearLogDelegate = InArgs._OnClearLog;
 }
 END_SLATE_FUNCTION_BUILD_OPTIMIZATION
 
@@ -1453,6 +1455,8 @@ void SOutputLog::OnClearLog()
 	MessagesTextMarshaller->ClearMessages();
 	MessagesTextBox->Refresh();
 	bIsUserScrolled = false;
+
+	[[maybe_unused]] bool bOnClearLogDelegateExecuted = OnClearLogDelegate.ExecuteIfBound();
 }
 
 void SOutputLog::OnHighlightCategory(FName NewCategoryToHighlight)
diff --git a/Engine/Source/Developer/OutputLog/Private/SOutputLog.h b/Engine/Source/Developer/OutputLog/Private/SOutputLog.h
--- a/Engine/Source/Developer/OutputLog/Private/SOutputLog.h
+++ b/Engine/Source/Developer/OutputLog/Private/SOutputLog.h
@@ -317,6 +317,8 @@ public:
 
 		SLATE_EVENT(FSimpleDelegate, OnCloseConsole)
 
+		SLATE_EVENT(FSimpleDelegate, OnClearLog)
+
 		/** All messages captured before this log window has been created */
 		SLATE_ARGUMENT( TArray< TSharedPtr<FOutputLogMessage> >, Messages )
 
@@ -418,6 +420,8 @@ protected:
 	/** True if the user has scrolled the window upwards */
 	bool bIsUserScrolled;
 
+	FSimpleDelegate OnClearLogDelegate; 
+
 private:
 
 	void BuildInitialLogCategoryFilter(const FArguments& InArgs);
diff --git a/Engine/Source/Developer/OutputLog/Public/OutputLogModule.h b/Engine/Source/Developer/OutputLog/Public/OutputLogModule.h
--- a/Engine/Source/Developer/OutputLog/Public/OutputLogModule.h
+++ b/Engine/Source/Developer/OutputLog/Public/OutputLogModule.h
@@ -41,7 +41,9 @@ public:
 	virtual void StartupModule();
 	virtual void ShutdownModule();
 
+	static const FName ModuleName;
 	static OUTPUTLOG_API FOutputLogModule& Get();
+	static OUTPUTLOG_API FOutputLogModule* TryGet();
 
 	/** Returns whether debug console widgets should be hidden */
 	virtual bool ShouldHideConsole() const;
@@ -90,6 +92,8 @@ public:
 	through OutputLogModule to avoid dependencies. */
 	virtual bool ShouldCycleToOutputLogDrawer() const;
 
+	TMulticastDelegate<void()> OnOutputLogDrawerCleared;
+
 private:
 	TSharedRef<SDockTab> SpawnOutputLogTab(const FSpawnTabArgs& Args);
 	void OnOutputLogTabClosed(TSharedRef<SDockTab> Tab);
