Fix category filtering in the Output Log when the "Show All" checkbox is enabled

FOutputLogFilter::IsMessageAllowed no longer includes the state of the "Show All" checkbox (bShowAllCategories) in the filtering logic; this restores the filtering logic wrt "Show All" to as it was prior to changelist 25554621 ( it preserves the addition of IgnoreFilterVerbosities from 25554621 ).

FOutputLogTextLayoutMarshaller::GetCategoryForLocation now extracts the category from the TextLayout.  The previous implementation was broken in the case of any applied filters because it was based on the unfiltered Messages array.

diff --git a/Engine/Source/Developer/OutputLog/Private/SOutputLog.cpp b/Engine/Source/Developer/OutputLog/Private/SOutputLog.cpp
--- a/Engine/Source/Developer/OutputLog/Private/SOutputLog.cpp
+++ b/Engine/Source/Developer/OutputLog/Private/SOutputLog.cpp
@@ -13,6 +13,7 @@
 #include "Framework/Text/SlateTextLayout.h"
 #include "Framework/Text/SlateTextRun.h"
 #include "Framework/Application/SlateApplication.h"
+#include "Internationalization/BreakIterator.h"
 #include "Widgets/Text/STextBlock.h"
 #include "Widgets/Input/SMenuAnchor.h"
 #include "Framework/MultiBox/MultiBoxBuilder.h"
@@ -1085,9 +1086,44 @@ void FOutputLogTextLayoutMarshaller::MarkMessagesCacheAsDirty()
 
 FName FOutputLogTextLayoutMarshaller::GetCategoryForLocation(const FTextLocation Location) const
 {
-	if (Messages.IsValidIndex(Location.GetLineIndex()))
+	if (TextLayout)
 	{
-		return Messages[Location.GetLineIndex()]->Category;
+		TSharedRef<IBreakIterator> WordBreakIterator{ FBreakIterator::CreateWordBreakIterator() };
+
+		int32 LineIndex = Location.GetLineIndex();
+
+		// A Message may be split across multiple lines in the TextLayout, so work backwards to find the Category on the first line of the message.
+		while (TextLayout->GetLineModels().IsValidIndex(LineIndex))
+		{
+			const FTextLayout::FLineModel& LineModel = TextLayout->GetLineModels()[LineIndex];	
+
+			WordBreakIterator->SetStringRef(&LineModel.Text.Get());
+
+			int32 PreviousBreak = WordBreakIterator->ResetToBeginning();
+			int32 CurrentBreak = 0;
+
+			// Iterate words starting from the beginning of the line, as the Category is one of the first words in a message.
+			while ((CurrentBreak = WordBreakIterator->MoveToNext()) != INDEX_NONE)
+			{
+				FTextSelection Selection{ FTextLocation(LineIndex, CurrentBreak), FTextLocation(LineIndex, PreviousBreak) };
+
+				FString SelectedText;
+				TextLayout->GetSelectionAsText(SelectedText, Selection);
+
+				FName PossibleCategory(SelectedText, FNAME_Find);
+
+				if (!PossibleCategory.IsNone() && Filter->IsLogCategoryAvailable(PossibleCategory))
+				{
+					return PossibleCategory;
+				}
+
+				PreviousBreak = CurrentBreak;
+			}
+
+			WordBreakIterator->ClearString();
+
+			LineIndex--;
+		}
 	}
 
 	return NAME_None;
@@ -2317,7 +2353,7 @@ bool FOutputLogFilter::IsMessageAllowed(const TSharedPtr<FOutputLogMessage>& Mes
 
 	// Filter by Category
 	{
-		if (!bShowAllCategories && !IgnoreFilterVerbosities.Contains(Message->Verbosity) && !IsLogCategoryEnabled(Message->Category))
+		if (!IgnoreFilterVerbosities.Contains(Message->Verbosity) && !IsLogCategoryEnabled(Message->Category))
 		{
 			return false;
 		}
@@ -2358,6 +2394,11 @@ void FOutputLogFilter::AddAvailableLogCategory(const FName& LogCategory)
 	}
 }
 
+bool FOutputLogFilter::IsLogCategoryAvailable(const FName& LogCategory) const
+{
+	return AvailableLogCategories.Contains(LogCategory);
+}
+
 void FOutputLogFilter::ToggleLogCategory(const FName& LogCategory)
 {
 	int32 FoundIndex = SelectedLogCategories.Find(LogCategory);
diff --git a/Engine/Source/Developer/OutputLog/Private/SOutputLog.h b/Engine/Source/Developer/OutputLog/Private/SOutputLog.h
--- a/Engine/Source/Developer/OutputLog/Private/SOutputLog.h
+++ b/Engine/Source/Developer/OutputLog/Private/SOutputLog.h
@@ -283,6 +283,9 @@ struct FOutputLogFilter
 	/** Adds a Log Category to the list of available categories, if it isn't already present */
 	void AddAvailableLogCategory(const FName& LogCategory);
 
+	/** Returns true if the specified log category is in the list of available categories */
+	bool IsLogCategoryAvailable(const FName& LogCategory) const;
+
 	/** Enables or disables a Log Category in the filter */
 	void ToggleLogCategory(const FName& LogCategory);
 
