PR #11813: Output log filtering controls
* Add shift modifier to skip setting the All Categories filter of the OutputLog

diff --git a/Engine/Source/Developer/OutputLog/Private/SOutputLog.cpp b/Engine/Source/Developer/OutputLog/Private/SOutputLog.cpp
--- a/Engine/Source/Developer/OutputLog/Private/SOutputLog.cpp
+++ b/Engine/Source/Developer/OutputLog/Private/SOutputLog.cpp
@@ -1845,25 +1845,37 @@ ECheckBoxState SOutputLog::VerbosityErrors_IsChecked() const
 	return ECheckBoxState::Unchecked;
 }
 
-void SOutputLog::VerbosityLogs_Execute()
-{ 
-	// Rotate through: showing the verbosity, showing the verbosity while ignoring filter categories, and hiding the verbosity
-	if (Filter.bShowLogs)
+static void VerbosityGeneric_Execute(bool& Flag, TSet<ELogVerbosity::Type>& Filter, ELogVerbosity::Type Verbosity)
+{
+	if (FSlateApplication::Get().GetModifierKeys().IsShiftDown())
 	{
-		if (Filter.IgnoreFilterVerbosities.Contains(ELogVerbosity::Log))
+		Flag = !Flag;
+	}
+	else
+	{
+		// Rotate through: showing the verbosity, showing the verbosity while ignoring filter categories, and hiding the verbosity
+		if (Flag)
 		{
-			Filter.bShowLogs = false;
-			Filter.IgnoreFilterVerbosities.Remove(ELogVerbosity::Log);
+			if (Filter.Contains(Verbosity))
+			{
+				Flag = false;
+				Filter.Remove(Verbosity);
+			}
+			else
+			{
+				Filter.Emplace(Verbosity);
+			}
 		}
 		else
 		{
-			Filter.IgnoreFilterVerbosities.Emplace(ELogVerbosity::Log);
+			Flag = true;
 		}
 	}
-	else
-	{
-		Filter.bShowLogs = true;
-	}
+}
+
+void SOutputLog::VerbosityLogs_Execute()
+{ 
+	VerbosityGeneric_Execute(Filter.bShowLogs, Filter.IgnoreFilterVerbosities, ELogVerbosity::Log);
 
 	// Flag the messages count as dirty
 	MessagesTextMarshaller->MarkMessagesCacheAsDirty();
@@ -1873,23 +1885,7 @@ void SOutputLog::VerbosityLogs_Execute()
 
 void SOutputLog::VerbosityWarnings_Execute()
 {
-	// Rotate through: showing the verbosity, showing the verbosity while ignoring filter categories, and hiding the verbosity
-	if (Filter.bShowWarnings)
-	{
-		if (Filter.IgnoreFilterVerbosities.Contains(ELogVerbosity::Warning))
-		{
-			Filter.bShowWarnings = false;
-			Filter.IgnoreFilterVerbosities.Remove(ELogVerbosity::Warning);
-		}
-		else
-		{
-			Filter.IgnoreFilterVerbosities.Emplace(ELogVerbosity::Warning);
-		}
-	}
-	else
-	{
-		Filter.bShowWarnings = true;
-	}
+	VerbosityGeneric_Execute(Filter.bShowWarnings, Filter.IgnoreFilterVerbosities, ELogVerbosity::Warning);
 
 	// Flag the messages count as dirty
 	MessagesTextMarshaller->MarkMessagesCacheAsDirty();
@@ -1899,23 +1895,7 @@ void SOutputLog::VerbosityWarnings_Execute()
 
 void SOutputLog::VerbosityErrors_Execute()
 {
-	// Rotate through: showing the verbosity, showing the verbosity while ignoring filter categories, and hiding the verbosity
-	if (Filter.bShowErrors)
-	{
-		if (Filter.IgnoreFilterVerbosities.Contains(ELogVerbosity::Error))
-		{
-			Filter.bShowErrors = false;
-			Filter.IgnoreFilterVerbosities.Remove(ELogVerbosity::Error);
-		}
-		else
-		{
-			Filter.IgnoreFilterVerbosities.Emplace(ELogVerbosity::Error);
-		}
-	}
-	else
-	{
-		Filter.bShowErrors = true;
-	}
+	VerbosityGeneric_Execute(Filter.bShowErrors, Filter.IgnoreFilterVerbosities, ELogVerbosity::Error);
 
 	// Flag the messages count as dirty
 	MessagesTextMarshaller->MarkMessagesCacheAsDirty();
