Add arguments to SOutputLog to control limiting the output log by the number of lines.  This is opt-in and depends on the console variable OutputLogModule.LimitLogging=true to enable in the right click menu of the output window.  By default the value is off and SwitchboardListener is the only app that enables the feature currently.

diff --git a/Engine/Source/Developer/OutputLog/Private/OutputLogModule.cpp b/Engine/Source/Developer/OutputLog/Private/OutputLogModule.cpp
--- a/Engine/Source/Developer/OutputLog/Private/OutputLogModule.cpp
+++ b/Engine/Source/Developer/OutputLog/Private/OutputLogModule.cpp
@@ -80,18 +80,32 @@ public:
 		return Messages;
 	}
 
+	void Suspend()
+	{
+		bSuspended = true;
+	}
+
+	void Resume()
+	{
+		bSuspended = false;
+	}
 protected:
 
 	virtual void Serialize(const TCHAR* V, ELogVerbosity::Type Verbosity, const class FName& Category) override
 	{
-		// Capture all incoming messages and store them in history
-		SOutputLog::CreateLogMessages(V, Verbosity, Category, Messages);
+		if (!bSuspended)
+		{
+			// Capture all incoming messages and store them in history
+			SOutputLog::CreateLogMessages(V, Verbosity, Category, Messages);
+		}
 	}
 
 private:
 
 	/** All log messsges since this module has been started */
 	TArray< TSharedPtr<FOutputLogMessage> > Messages;
+
+	bool bSuspended = false;
 };
 
 
@@ -252,7 +266,22 @@ TSharedRef<SWidget> FOutputLogModule::MakeOutputLogWidget(const FOutputLogCreati
 			.Messages(OutputLogHistory->GetMessages())
 			.SettingsMenuFlags(Params.SettingsMenuCreationFlags)
 			.DefaultCategorySelection(Params.DefaultCategorySelection)
-			.AllowInitialLogCategory(Params.AllowAsInitialLogCategory);
+			.AllowInitialLogCategory(Params.AllowAsInitialLogCategory)
+			.EnableLoggingLimitMenu(Params.LoggingLimit.IsSet())
+			.LoggingLineLimit(Params.LoggingLimit);
+}
+
+void FOutputLogModule::SuspendGlobalLog()
+{
+	check(OutputLogHistory);
+	OutputLogHistory->Suspend();
+}
+
+
+void FOutputLogModule::ResumeGlobalLog()
+{
+	check(OutputLogHistory);
+	OutputLogHistory->Resume();
 }
 
 void FOutputLogModule::ToggleDebugConsoleForWindow(const TSharedRef<SWindow>& Window, const EDebugConsoleStyle::Type InStyle, const FDebugConsoleDelegates& DebugConsoleDelegates)
@@ -475,4 +504,4 @@ bool FOutputLogModule::ShouldCycleToOutputLogDrawer() const
 #else
 	return true;
 #endif
-}
\ No newline at end of file
+}
diff --git a/Engine/Source/Developer/OutputLog/Private/SOutputLog.cpp b/Engine/Source/Developer/OutputLog/Private/SOutputLog.cpp
--- a/Engine/Source/Developer/OutputLog/Private/SOutputLog.cpp
+++ b/Engine/Source/Developer/OutputLog/Private/SOutputLog.cpp
@@ -2,6 +2,7 @@
 
 #include "SOutputLog.h"
 #include "ConsoleSettings.h"
+#include "Framework/Commands/UICommandInfo.h"
 #include "Framework/Text/IRun.h"
 #include "Framework/Text/TextLayout.h"
 #include "Misc/ConfigCacheIni.h"
@@ -23,6 +24,8 @@
 #include "Widgets/Input/SSearchBox.h"
 #include "Widgets/Input/SCheckBox.h"
 #include "Widgets/Images/SImage.h"
+#include "Widgets/Input/SNumericEntryBox.h"
+#include "Widgets/Layout/SSpacer.h"
 #include "Features/IModularFeatures.h"
 #include "Misc/CoreDelegates.h"
 #include "HAL/PlatformOutputDevices.h"
@@ -1030,6 +1033,7 @@ void FOutputLogTextLayoutMarshaller::ClearMessages()
 {
 	NextPendingMessageIndex = 0;
 	Messages.Empty();
+	bNumMessagesCacheDirty = true;
 	MakeDirty();
 }
 
@@ -1079,6 +1083,17 @@ int32 FOutputLogTextLayoutMarshaller::GetNumFilteredMessages()
 	return CachedNumMessages;
 }
 
+int32 FOutputLogTextLayoutMarshaller::GetNumCachedMessages()
+{
+	// Re-count messages if filter changed before we refresh
+	if (bNumMessagesCacheDirty)
+	{
+		CountMessages();
+	}
+
+	return CachedNumMessages;
+}
+
 void FOutputLogTextLayoutMarshaller::MarkMessagesCacheAsDirty()
 {
 	bNumMessagesCacheDirty = true;
@@ -1165,6 +1180,10 @@ void SOutputLog::Construct( const FArguments& InArgs, bool bCreateDrawerDockButt
 	bShouldCreateDrawerDockButton = bCreateDrawerDockButton;
 	BuildInitialLogCategoryFilter(InArgs);
 
+	bShouldShowLoggingLimitMenu = InArgs._EnableLoggingLimitMenu;
+	bEnableLoggingLimit = InArgs._LoggingLineLimit.IsSet();
+	LoggingLineLimit = InArgs._LoggingLineLimit.Get(10000);
+
 	MessagesTextMarshaller = FOutputLogTextLayoutMarshaller::Create(InArgs._Messages, &Filter);
 
 	MessagesTextBox = SNew(SMultiLineEditableTextBox)
@@ -1229,6 +1248,22 @@ void SOutputLog::Construct( const FArguments& InArgs, bool bCreateDrawerDockButt
 				]
 			]
 			+SHorizontalBox::Slot()
+			.HAlign(HAlign_Left)
+			.VAlign(VAlign_Center)
+			.Padding(4, 0)
+			[
+				SNew(STextBlock)
+				.Text(LOCTEXT("LogLineLimitReached", "Log line limit reached. Clear log to continue."))
+				.ColorAndOpacity(FSlateColor(FLinearColor::Yellow))
+				.Visibility(MakeAttributeLambda([this]() {
+					if (!bEnableLoggingLimit || MessagesTextMarshaller->GetNumCachedMessages() < LoggingLineLimit)
+					{
+						return EVisibility::Hidden;
+					}
+					return EVisibility::Visible;
+				}))
+			]
+			+SHorizontalBox::Slot()
 			.HAlign(HAlign_Right)
 			.VAlign(VAlign_Center)
 			.Padding(4, 0)
@@ -1278,7 +1313,9 @@ void SOutputLog::Construct( const FArguments& InArgs, bool bCreateDrawerDockButt
 		.AutoHeight()
 		[
 			SAssignNew(ConsoleInputBox, SConsoleInputBox)
-			.Visibility(MakeAttributeLambda([]() { return  FOutputLogModule::Get().ShouldHideConsole() ? EVisibility::Collapsed : EVisibility::Visible; }))
+			.Visibility(MakeAttributeLambda([]() {
+				return FOutputLogModule::Get().ShouldHideConsole() ? EVisibility::Collapsed : EVisibility::Visible;
+			}))
 			.OnConsoleCommandExecuted(this, &SOutputLog::OnConsoleCommandExecuted)
 			.OnCloseConsole(InArgs._OnCloseConsole)
 			// Always place suggestions above the input line for the output log widget
@@ -1421,7 +1458,44 @@ bool SOutputLog::CreateLogMessages( const TCHAR* V, ELogVerbosity::Type Verbosit
 
 void SOutputLog::Serialize(const TCHAR* V, ELogVerbosity::Type Verbosity, const class FName& Category)
 {
-	MessagesTextMarshaller->AppendPendingMessage(V, Verbosity, Category);
+	if (!bEnableLoggingLimit || MessagesTextMarshaller->GetNumCachedMessages() < LoggingLineLimit)
+	{
+		MessagesTextMarshaller->AppendPendingMessage(V, Verbosity, Category);
+	}
+}
+
+TSharedRef<SWidget> SOutputLog::MakeLogLimitMenuItem()
+{
+	return SNew(SHorizontalBox)
+		+SHorizontalBox::Slot()
+		.AutoWidth()
+		.VAlign(VAlign_Center)
+		.HAlign(HAlign_Center)
+		[
+			SNew(STextBlock)
+			.Text(LOCTEXT("LimitLog", "Logging Limit"))
+		]
+		+SHorizontalBox::Slot()
+		.FillWidth(1.f)
+		[
+			SNew(SSpacer)
+		]
+		+SHorizontalBox::Slot()
+		.HAlign(HAlign_Right)
+		[
+			SNew(SNumericEntryBox<int32>)
+			.AllowSpin(true)
+			.Justification(ETextJustify::Right)
+			.MinDesiredValueWidth(100)
+			.MaxSliderValue(100000)
+			.OnValueChanged_Lambda([this](int32 NewValue){
+				if (NewValue > 100)
+				{
+					LoggingLineLimit = NewValue;
+				}
+			})
+			.Value_Lambda([this](){ return LoggingLineLimit; })
+		];
 }
 
 void SOutputLog::ExtendTextBoxMenu(FMenuBuilder& Builder)
@@ -1438,6 +1512,18 @@ void SOutputLog::ExtendTextBoxMenu(FMenuBuilder& Builder)
 		ClearOutputLogAction
 		);
 
+	Builder.AddMenuEntry(
+		FUIAction(
+			FExecuteAction::CreateLambda([this](){ bEnableLoggingLimit = !bEnableLoggingLimit; }),
+			FCanExecuteAction::CreateLambda([] { return true; }),
+			FIsActionChecked::CreateLambda([this] { return bEnableLoggingLimit; }),
+			FIsActionButtonVisible::CreateLambda([this] { return bShouldShowLoggingLimitMenu; })
+		),
+		MakeLogLimitMenuItem(),
+		NAME_None,
+		LOCTEXT("LimitLogToolTip", "Limits Logging to specified number of lines."),
+		EUserInterfaceActionType::ToggleButton);
+
 	const FVector2D CursorPos = FSlateApplication::Get().GetCursorPos();
 	const FVector2D RelativeCursorPos = MessagesTextBox->GetTickSpaceGeometry().AbsoluteToLocal(CursorPos);
 	const FTextLocation CursorTextLocation = MessagesTextMarshaller->GetTextLocationAt(RelativeCursorPos);
diff --git a/Engine/Source/Developer/OutputLog/Private/SOutputLog.h b/Engine/Source/Developer/OutputLog/Private/SOutputLog.h
--- a/Engine/Source/Developer/OutputLog/Private/SOutputLog.h
+++ b/Engine/Source/Developer/OutputLog/Private/SOutputLog.h
@@ -328,6 +328,12 @@ public:
 		/**  */
 		SLATE_ARGUMENT_DEFAULT( EOutputLogSettingsMenuFlags, SettingsMenuFlags ) = EOutputLogSettingsMenuFlags::None;
 
+		/** Should this output log enable support for limiting the number of logged lines. */
+		SLATE_ARGUMENT_DEFAULT( bool, EnableLoggingLimitMenu ) = false;
+
+		/** The limit to the number of lines we output to the logging widget. */
+		SLATE_ARGUMENT_DEFAULT( TOptional<int32>, LoggingLineLimit ) = {};
+
 		SLATE_ARGUMENT( FDefaultCategorySelectionMap, DefaultCategorySelection )
 
 		/** Used to determine the set of initially discovered log categories that should be selected */
@@ -427,6 +433,9 @@ protected:
 
 private:
 
+	/** The log limit menu entry to prevent logs from consuming the output log widget. */
+	TSharedRef<SWidget> MakeLogLimitMenuItem();
+
 	void BuildInitialLogCategoryFilter(const FArguments& InArgs);
 	
 	/** Called by Slate when the filter box changes text. */
@@ -507,6 +516,7 @@ private:
 	void OpenLogFileInExternalEditor();
 
 	FReply OnDockInLayoutClicked();
+
 protected:
 	TSharedPtr<SConsoleInputBox> ConsoleInputBox;
 
@@ -516,6 +526,12 @@ protected:
 	FDelegateHandle SettingsWatchHandle;
 
 	bool bShouldCreateDrawerDockButton = false;
+
+	bool bShouldShowLoggingLimitMenu = false;
+
+	bool bEnableLoggingLimit = false;
+
+	int32 LoggingLineLimit = 10000;
 };
 
 /** Output log text marshaller to convert an array of FOutputLogMessages into styled lines to be consumed by an FTextLayout */
@@ -539,6 +555,7 @@ public:
 
 	int32 GetNumMessages() const;
 	int32 GetNumFilteredMessages();
+	int32 GetNumCachedMessages();
 
 	void MarkMessagesCacheAsDirty();
 
diff --git a/Engine/Source/Developer/OutputLog/Public/OutputLogCreationParams.h b/Engine/Source/Developer/OutputLog/Public/OutputLogCreationParams.h
--- a/Engine/Source/Developer/OutputLog/Public/OutputLogCreationParams.h
+++ b/Engine/Source/Developer/OutputLog/Public/OutputLogCreationParams.h
@@ -38,6 +38,10 @@ struct FOutputLogCreationParams
 
 	/** Maps each log category to whether it should be selected or deselected by default. The caller is responsible to enter valid category names. */
 	FDefaultCategorySelectionMap DefaultCategorySelection;
-	
+
+	/** Called when the console tab is closed. */
 	FSimpleDelegate OnCloseConsole;
-};
\ No newline at end of file
+
+	/** Do we want to support limited console output with a predefined limit. */
+	TOptional<int32> LoggingLimit;
+};
diff --git a/Engine/Source/Developer/OutputLog/Public/OutputLogModule.h b/Engine/Source/Developer/OutputLog/Public/OutputLogModule.h
--- a/Engine/Source/Developer/OutputLog/Public/OutputLogModule.h
+++ b/Engine/Source/Developer/OutputLog/Public/OutputLogModule.h
@@ -55,7 +55,7 @@ public:
 	virtual TSharedRef<SWidget> MakeOutputLogDrawerWidget(const FSimpleDelegate& OnCloseConsole);
 
 	virtual TSharedRef<SWidget> MakeOutputLogWidget(const FOutputLogCreationParams& Params);
-	
+
 	/** Opens a debug console in the specified window, if not already open */
 	virtual void ToggleDebugConsoleForWindow(const TSharedRef<SWindow>& Window, const EDebugConsoleStyle::Type InStyle, const FDebugConsoleDelegates& DebugConsoleDelegates);
 
@@ -68,6 +68,16 @@ public:
 
 	virtual const TSharedPtr<SWidget> GetOutputLog() const;
 
+	/*
+	 * Suspend the collection of global log data that is used for constructing output log widgets.  This global log will perpetually collect logging information perpetually and may exhaust memory.
+	 */
+	virtual void SuspendGlobalLog();
+
+	/*
+	 * Resume global log collection. Note this has a memory impact as log entries are stored for the life of the running edit.
+	 */
+	virtual void ResumeGlobalLog();
+
 	/** Opens and focuses on the Output Log Drawer if the status bar exists, otherwise opens and focuses on the Output Log Tab. */
 	void OUTPUTLOG_API FocusOutputLog();
 
diff --git a/Engine/Source/Programs/SwitchboardListener/SblSlate/Private/SblMainWindow.cpp b/Engine/Source/Programs/SwitchboardListener/SblSlate/Private/SblMainWindow.cpp
--- a/Engine/Source/Programs/SwitchboardListener/SblSlate/Private/SblMainWindow.cpp
+++ b/Engine/Source/Programs/SwitchboardListener/SblSlate/Private/SblMainWindow.cpp
@@ -549,7 +549,12 @@ TSharedRef<SWidget> FSwitchboardListenerMainWindow::CreateOutputLog()
 		return LogCategory == LogSwitchboard.GetCategoryName();
 	});
 
-	return OutputLogModule.MakeOutputLogWidget(OutputLogCreationParams);
+	// Limit logging output to 10,000 so that we don't over consume on memory.
+	OutputLogCreationParams.LoggingLimit = 10000;
+	TSharedRef<SWidget> OutputLog = OutputLogModule.MakeOutputLogWidget(OutputLogCreationParams);
+	OutputLogModule.SuspendGlobalLog();
+	return OutputLog;
+
 }
 
 
