- Added support to OutputLog to only have the TargetDevice functionality if DeveloperTools are included in the target (so programs can use it without needing to pull in the whole mess that is TargetPlatform)

diff --git a/Engine/Source/Developer/OutputLog/OutputLog.Build.cs b/Engine/Source/Developer/OutputLog/OutputLog.Build.cs
--- a/Engine/Source/Developer/OutputLog/OutputLog.Build.cs
+++ b/Engine/Source/Developer/OutputLog/OutputLog.Build.cs
@@ -14,7 +14,6 @@ public class OutputLog : ModuleRules
 				"InputCore",
 				"Slate",
 				"SlateCore",
-				"TargetPlatform",
 				"DesktopPlatform",
 				"ToolWidgets",
 				"ToolMenus",
@@ -34,11 +33,13 @@ public class OutputLog : ModuleRules
 
 		if (Target.bBuildEditor || Target.bBuildDeveloperTools)
 		{
-			PrivateIncludePathModuleNames.AddRange(
-				new string[] {
-					"WorkspaceMenuStructure",
-				}
-			);
+			PrivateIncludePathModuleNames.Add("WorkspaceMenuStructure");
+			PrivateDependencyModuleNames.Add("TargetPlatform");
+			PrivateDefinitions.Add("OUTPUTLOG_HAS_TARGET_PLATFORMS=1");
+		}
+		else
+		{
+			PrivateDefinitions.Add("OUTPUTLOG_HAS_TARGET_PLATFORMS=0");
 		}
 
 		if (Target.bCompileAgainstEngine)
diff --git a/Engine/Source/Developer/OutputLog/Private/SDeviceOutputLog.cpp b/Engine/Source/Developer/OutputLog/Private/SDeviceOutputLog.cpp
--- a/Engine/Source/Developer/OutputLog/Private/SDeviceOutputLog.cpp
+++ b/Engine/Source/Developer/OutputLog/Private/SDeviceOutputLog.cpp
@@ -9,23 +9,29 @@
 #include "Framework/Commands/UIAction.h"
 #include "Widgets/Input/SMultiLineEditableTextBox.h"
 #include "Framework/MultiBox/MultiBoxBuilder.h"
+#if OUTPUTLOG_HAS_TARGET_PLATFORMS
 #include "Interfaces/ITargetPlatform.h"
 #include "Interfaces/ITargetPlatformManagerModule.h"
+#endif
 #include "PlatformInfo.h"
 #include "OutputLogModule.h"
 #include "OutputLogStyle.h"
 #include "SSimpleComboButton.h"
 
+#if OUTPUTLOG_HAS_TARGET_PLATFORMS
 static bool IsSupportedPlatform(ITargetPlatform* Platform)
 {
 	check(Platform);
 	return Platform->SupportsFeature( ETargetPlatformFeatures::DeviceOutputLog );
 }
+#endif
 
 
 void SDeviceOutputLog::Construct( const FArguments& InArgs )
 {
+#if OUTPUTLOG_HAS_TARGET_PLATFORMS
 	bAutoSelectDevice = InArgs._AutoSelectDevice;
+#endif
 
 	MessagesTextMarshaller = FOutputLogTextLayoutMarshaller::Create(TArray<TSharedPtr<FOutputLogMessage>>(), &Filter);
 
@@ -48,6 +54,7 @@ void SDeviceOutputLog::Construct( const FArguments& InArgs )
 			[
 				MessagesTextBox.ToSharedRef()
 			]
+#if OUTPUTLOG_HAS_TARGET_PLATFORMS
 			// The console input box
 			+SVerticalBox::Slot()
 			.AutoHeight()
@@ -99,11 +106,13 @@ void SDeviceOutputLog::Construct( const FArguments& InArgs )
 					.SuggestionListPlacement( MenuPlacement_AboveAnchor )
 				]
 			]
+#endif
 	];
 
 	bIsUserScrolled = false;
 	RequestForceScroll();
 	
+#if OUTPUTLOG_HAS_TARGET_PLATFORMS
 	ITargetPlatformControls::OnDeviceDiscovered().AddRaw(this, &SDeviceOutputLog::HandleTargetPlatformDeviceDiscovered);
 	ITargetPlatformControls::OnDeviceLost().AddRaw(this, &SDeviceOutputLog::HandleTargetPlatformDeviceLost);
 		
@@ -124,20 +133,24 @@ void SDeviceOutputLog::Construct( const FArguments& InArgs )
 			}
 		}
 	}
+#endif
 }
 
 SDeviceOutputLog::~SDeviceOutputLog()
 {
+#if OUTPUTLOG_HAS_TARGET_PLATFORMS
 	ITargetPlatformControls::OnDeviceDiscovered().RemoveAll(this);
 	ITargetPlatformControls::OnDeviceLost().RemoveAll(this);
 
 	// Clearing the pointer manually to ensure that when the pointed device output object is destroyed
 	// SDeviceOutputLog is still in a valid state in case CurrentDeviceOutputPtr wanted to dereference it.
 	CurrentDeviceOutputPtr.Reset();
+#endif
 }
 
 void SDeviceOutputLog::Tick(const FGeometry& AllottedGeometry, const double InCurrentTime, const float InDeltaTime)
 {
+#if OUTPUTLOG_HAS_TARGET_PLATFORMS
 	// If auto-select is enabled request connecting to the default device and select it
 	if (!CurrentDevicePtr.IsValid() && bAutoSelectDevice)
 	{
@@ -164,6 +177,7 @@ void SDeviceOutputLog::Tick(const FGeometry& AllottedGeometry, const double InCu
 			CurrentDeviceOutputPtr = PinnedPtr->CreateDeviceOutputRouter(this);
 		}
 	}
+#endif
 
 	FScopeLock ScopeLock(&BufferedLinesSynch);
 	if (BufferedLines.Num() > 0)
@@ -191,6 +205,7 @@ bool SDeviceOutputLog::CanBeUsedOnAnyThread() const
 
 void SDeviceOutputLog::ExecuteConsoleCommand(const FString& ExecCommand)
 {
+#if OUTPUTLOG_HAS_TARGET_PLATFORMS
 	if (CurrentDevicePtr.IsValid())
 	{
 		ITargetDevicePtr PinnedPtr = CurrentDevicePtr->DeviceWeakPtr.Pin();
@@ -199,8 +214,10 @@ void SDeviceOutputLog::ExecuteConsoleCommand(const FString& ExecCommand)
 			PinnedPtr->ExecuteConsoleCommand(ExecCommand);
 		}
 	}
+#endif
 }
 
+#if OUTPUTLOG_HAS_TARGET_PLATFORMS
 void SDeviceOutputLog::HandleTargetPlatformDeviceLost(ITargetDeviceRef LostDevice)
 {
 	FTargetDeviceId LostDeviceId = LostDevice->GetId();
@@ -403,3 +420,4 @@ FText SDeviceOutputLog::GetSelectedTargetDeviceText() const
 {
 	return GetTargetDeviceText(CurrentDevicePtr);
 }
+#endif
diff --git a/Engine/Source/Developer/OutputLog/Private/SDeviceOutputLog.h b/Engine/Source/Developer/OutputLog/Private/SDeviceOutputLog.h
--- a/Engine/Source/Developer/OutputLog/Private/SDeviceOutputLog.h
+++ b/Engine/Source/Developer/OutputLog/Private/SDeviceOutputLog.h
@@ -8,11 +8,16 @@
 #include "Widgets/SWidget.h"
 #include "Widgets/DeclarativeSyntaxSupport.h"
 #include "SOutputLog.h"
+
+#if OUTPUTLOG_HAS_TARGET_PLATFORMS
 #include "Interfaces/TargetDeviceId.h"
 #include "Interfaces/ITargetDevice.h"
+#endif
+
 
 class SComboButton;
 
+#if OUTPUTLOG_HAS_TARGET_PLATFORMS
 struct FTargetDeviceEntry
 {
 	FTargetDeviceId			DeviceId;
@@ -23,6 +28,7 @@ struct FTargetDeviceEntry
 typedef TSharedPtr<FTargetDeviceEntry> FTargetDeviceEntryPtr;
 
 DECLARE_DELEGATE_OneParam(FSelectedTargetDeviceChangedDelegate, ITargetDevicePtr);
+#endif
 
 class SDeviceOutputLog : public SOutputLog
 {
@@ -43,8 +49,10 @@ public:
 	 */
 	void Construct( const FArguments& InArgs );
 
+#if OUTPUTLOG_HAS_TARGET_PLATFORMS
 	FSelectedTargetDeviceChangedDelegate& OnSelectedDeviceChanged() { return OnSelectedDeviceChangedDelegate; }
 	ITargetDevicePtr GetSelectedTargetDevice() const;
+#endif
 
 protected:
 	// SWidget interface
@@ -56,6 +64,7 @@ protected:
 
 	void ExecuteConsoleCommand(const FString& ExecCommand);
 
+#if OUTPUTLOG_HAS_TARGET_PLATFORMS
 	/** Callback for lost target devices. */
 	void HandleTargetPlatformDeviceLost(ITargetDeviceRef LostDevice);
 	/** Callback for discovered target devices. */
@@ -73,19 +82,23 @@ protected:
 	
 	const FSlateBrush* GetTargetDeviceBrush(FTargetDeviceEntryPtr DeviceEntry) const;
 	const FSlateBrush* GetSelectedTargetDeviceBrush() const;
-		
+#endif
+
 private:
+#if OUTPUTLOG_HAS_TARGET_PLATFORMS
 	TArray<FTargetDeviceEntryPtr>			DeviceList;
 	FTargetDeviceEntryPtr					CurrentDevicePtr;
 	ITargetDeviceOutputPtr					CurrentDeviceOutputPtr;
 
 	TSharedPtr<SComboButton> TargetDeviceComboButton;
-		
+#endif
 	/** Synchronization object for access to buffered lines */
 	FCriticalSection		BufferedLinesSynch;
 	TArray<FBufferedLine>	BufferedLines;
 
+#if OUTPUTLOG_HAS_TARGET_PLATFORMS
 	FSelectedTargetDeviceChangedDelegate OnSelectedDeviceChangedDelegate;
 
 	bool bAutoSelectDevice;
+#endif
 };
