Added console command tooltips when displaying the history
Also, for the console command text box, there is now a tooltip for the command currently being typed (if it's valid), which also displays the current value for CVars (to avoid typing the variable without parameter, switch to the log window to peak at the current value and back to the console command to enter a new value)
Also added tooltips forcefields for both the edit text box and the suggestion list view

diff --git a/Engine/Plugins/Developer/Concert/ConcertApp/MultiUserServer/Source/MultiUserServer/Private/ConcertConsoleCommandExecutor.cpp b/Engine/Plugins/Developer/Concert/ConcertApp/MultiUserServer/Source/MultiUserServer/Private/ConcertConsoleCommandExecutor.cpp
--- a/Engine/Plugins/Developer/Concert/ConcertApp/MultiUserServer/Source/MultiUserServer/Private/ConcertConsoleCommandExecutor.cpp
+++ b/Engine/Plugins/Developer/Concert/ConcertApp/MultiUserServer/Source/MultiUserServer/Private/ConcertConsoleCommandExecutor.cpp
@@ -39,18 +39,24 @@ namespace UE::MultiUserServer
 	{
 		auto OnConsoleVariable = [&Out](const TCHAR *Name, IConsoleObject* CVar)
 		{
-			if (CVar->TestFlags(ECVF_Unregistered))
+			if (CVar->IsEnabled())
 			{
-				return;
+				Out.Add(FConsoleSuggestion(Name, CVar->GetDetailedHelp().ToString()));
 			}
-
-			Out.Add(FConsoleSuggestion(Name, CVar->GetHelp()));
 		};
 
-		IConsoleManager::Get().ForEachConsoleObjectThatContains(FConsoleObjectVisitor::CreateLambda(OnConsoleVariable), Input);
+		IConsoleManager& ConsoleManager = IConsoleManager::Get();
+		ConsoleManager.ForEachConsoleObjectThatContains(FConsoleObjectVisitor::CreateLambda(OnConsoleVariable), Input);
 		for (const FString& Name : GetDefault<UConsoleSettings>()->GetFilteredManualAutoCompleteCommands(Input))
 		{
-			Out.Add(FConsoleSuggestion(Name, FString()));
+			FString HelpString;
+			// Try to find a console object for this entry in order to retrieve a help string if possible :
+			const TCHAR* NamePtr = *Name;
+			if (IConsoleObject* CObj = ConsoleManager.FindConsoleObject(*FParse::Token(NamePtr, /*UseEscape = */false), /*bTrackFrequentCalls = */false); CObj && CObj->IsEnabled())
+			{
+				HelpString = CObj->GetDetailedHelp().ToString();
+			}
+			Out.Add(FConsoleSuggestion(Name, HelpString));
 		}
 	}
 
diff --git a/Engine/Plugins/Editor/ConsoleVariablesEditor/Source/ConsoleVariablesEditor/Public/Views/Widgets/SConsoleVariablesEditorCustomConsoleInputBox.cpp b/Engine/Plugins/Editor/ConsoleVariablesEditor/Source/ConsoleVariablesEditor/Public/Views/Widgets/SConsoleVariablesEditorCustomConsoleInputBox.cpp
--- a/Engine/Plugins/Editor/ConsoleVariablesEditor/Source/ConsoleVariablesEditor/Public/Views/Widgets/SConsoleVariablesEditorCustomConsoleInputBox.cpp
+++ b/Engine/Plugins/Editor/ConsoleVariablesEditor/Source/ConsoleVariablesEditor/Public/Views/Widgets/SConsoleVariablesEditorCustomConsoleInputBox.cpp
@@ -262,18 +262,10 @@ void SConsoleVariablesEditorCustomConsoleInputBox::OnInputTextChanged(const FTex
 		
 		auto OnConsoleVariable = [&AutoCompleteList](const TCHAR *Name, IConsoleObject* CVar)
 		{
-#if (UE_BUILD_SHIPPING || UE_BUILD_TEST)
-			if (CVar->TestFlags(ECVF_Cheat))
-			{
-				return;
-			}
-#endif // (UE_BUILD_SHIPPING || UE_BUILD_TEST)
-			if (CVar->TestFlags(ECVF_Unregistered))
-			{
-				return;
+			if (CVar->IsEnabled())
+			{	
+				AutoCompleteList.Add(Name);
 			}
-
-			AutoCompleteList.Add(Name);
 		};
 
 		IConsoleManager::Get().ForEachConsoleObjectThatContains(FConsoleObjectVisitor::CreateLambda(OnConsoleVariable), *InputTextStr);
diff --git a/Engine/Plugins/FX/Niagara/Source/NiagaraEditor/Private/Customizations/NiagaraPlatformSetCustomization.cpp b/Engine/Plugins/FX/Niagara/Source/NiagaraEditor/Private/Customizations/NiagaraPlatformSetCustomization.cpp
--- a/Engine/Plugins/FX/Niagara/Source/NiagaraEditor/Private/Customizations/NiagaraPlatformSetCustomization.cpp
+++ b/Engine/Plugins/FX/Niagara/Source/NiagaraEditor/Private/Customizations/NiagaraPlatformSetCustomization.cpp
@@ -1472,13 +1472,7 @@ void SNiagaraConsoleInputBox::OnTextChanged(const FText& InText)
 
 		auto OnConsoleVariable = [&AutoCompleteList](const TCHAR* Name, IConsoleObject* CVar)
 		{
-#if (UE_BUILD_SHIPPING || UE_BUILD_TEST)
-			if (CVar->TestFlags(ECVF_Cheat))
-			{
-				return;
-			}
-#endif // (UE_BUILD_SHIPPING || UE_BUILD_TEST)
-			if (CVar->TestFlags(ECVF_Unregistered))
+			if (!CVar->IsEnabled())
 			{
 				return;
 			}
diff --git a/Engine/Source/Developer/OutputLog/Private/SOutputLog.cpp b/Engine/Source/Developer/OutputLog/Private/SOutputLog.cpp
--- a/Engine/Source/Developer/OutputLog/Private/SOutputLog.cpp
+++ b/Engine/Source/Developer/OutputLog/Private/SOutputLog.cpp
@@ -231,6 +231,7 @@ void SConsoleInputBox::Construct(const FArguments& InArgs)
 					.OnIsTypedCharValid(FOnIsTypedCharValid::CreateLambda([](const TCHAR InCh) { return true; })) // allow tabs to be typed into the field
 					.ClearKeyboardFocusOnCommit(false)
 					.ModiferKeyForNewLine(EModifierKey::Shift)
+					.ToolTipText(this, &SConsoleInputBox::GetInputHelpText)
 				]
 			]
 		]
@@ -254,6 +255,10 @@ void SConsoleInputBox::Construct(const FArguments& InArgs)
 			]
 		)
 	];
+
+	// Don't let tooltips appear on top of the text box since it hampers visibility while typing the command : 
+	InputText->EnableToolTipForceField(true);
+	SuggestionListView->EnableToolTipForceField(true);
 }
 END_SLATE_FUNCTION_BUILD_OPTIMIZATION
 
@@ -353,18 +358,10 @@ void SConsoleInputBox::OnTextChanged(const FText& InText)
 		{
 			auto OnConsoleVariable = [&AutoCompleteList](const TCHAR *Name, IConsoleObject* CVar)
 			{
-#if (UE_BUILD_SHIPPING || UE_BUILD_TEST)
-				if (CVar->TestFlags(ECVF_Cheat))
-				{
-					return;
-				}
-#endif // (UE_BUILD_SHIPPING || UE_BUILD_TEST)
-				if (CVar->TestFlags(ECVF_Unregistered))
+				if (CVar->IsEnabled())
 				{
-					return;
+					AutoCompleteList.Add(FConsoleSuggestion(Name, CVar->GetDetailedHelp().ToString()));
 				}
-
-				AutoCompleteList.Add(FConsoleSuggestion(Name, CVar->GetHelp()));
 			};
 
 			IConsoleManager::Get().ForEachConsoleObjectThatContains(FConsoleObjectVisitor::CreateLambda(OnConsoleVariable), *InputTextStr);
@@ -482,6 +479,7 @@ FReply SConsoleInputBox::OnPreviewKeyDown(const FGeometry& MyGeometry, const FKe
 			const bool bShowHistory = InputText->GetText().IsEmpty() || KeyEvent.IsControlDown();
 			if (bShowHistory)
 			{
+				IConsoleManager& ConsoleManager = IConsoleManager::Get();
 				TArray<FString> HistoryNames;
 				if (ActiveCommandExecutor)
 				{
@@ -489,12 +487,20 @@ FReply SConsoleInputBox::OnPreviewKeyDown(const FGeometry& MyGeometry, const FKe
 				}
 				else
 				{
-					IConsoleManager::Get().GetConsoleHistory(TEXT(""), HistoryNames);
+					ConsoleManager.GetConsoleHistory(TEXT(""), HistoryNames);
 				}
 				TArray<FConsoleSuggestion> History;
 				for (const FString& Name : HistoryNames)
 				{
-					History.Add(FConsoleSuggestion(Name, FString()));
+					FString HelpString;
+					// Try to find a console object for this history entry in order to retrieve a help string if possible :
+					const TCHAR* NamePtr = *Name;
+					if (IConsoleObject* CObj = ConsoleManager.FindConsoleObject(*FParse::Token(NamePtr, /*UseEscape = */false), /*bTrackFrequentCalls = */false); CObj && CObj->IsEnabled())
+					{
+						HelpString = CObj->GetDetailedHelp().ToString();
+					}
+
+					History.Add(FConsoleSuggestion(Name, HelpString));
 				}
 				SetSuggestions(History, FText::GetEmpty());
 				
@@ -681,6 +687,22 @@ bool SConsoleInputBox::GetActiveCommandExecutorAllowMultiLine() const
 	return false;
 }
 
+FText SConsoleInputBox::GetInputHelpText() const
+{
+	const FString& InputTextStr = InputText->GetText().ToString();
+	if (!InputTextStr.IsEmpty())
+	{
+		// Try to find a console object for this entry in order to retrieve a help string if possible :
+		IConsoleManager& ConsoleManager = IConsoleManager::Get();
+		const TCHAR* InputTextStrPtr = *InputTextStr;
+		if (IConsoleObject* CObj = ConsoleManager.FindConsoleObject(*FParse::Token(InputTextStrPtr, /*UseEscape = */false), /*bTrackFrequentCalls = */false); CObj && CObj->IsEnabled())
+		{
+			return CObj->GetDetailedHelp();
+		}
+	}
+	return FText::GetEmpty();
+}
+
 bool SConsoleInputBox::IsCommandExecutorMenuEnabled() const
 {
 	return !ConsoleCommandCustomExec.IsBound(); // custom execs always show the default executor in the UI (which has the selector disabled)
diff --git a/Engine/Source/Developer/OutputLog/Private/SOutputLog.h b/Engine/Source/Developer/OutputLog/Private/SOutputLog.h
--- a/Engine/Source/Developer/OutputLog/Private/SOutputLog.h
+++ b/Engine/Source/Developer/OutputLog/Private/SOutputLog.h
@@ -135,6 +135,8 @@ protected:
 
 	bool GetActiveCommandExecutorAllowMultiLine() const;
 
+	FText GetInputHelpText() const;
+
 	bool IsCommandExecutorMenuEnabled() const;
 
 	TSharedRef<SWidget> GetCommandExecutorMenuContent();
diff --git a/Engine/Source/Editor/UnrealEd/Private/Toolkits/FConsoleCommandExecutor.cpp b/Engine/Source/Editor/UnrealEd/Private/Toolkits/FConsoleCommandExecutor.cpp
--- a/Engine/Source/Editor/UnrealEd/Private/Toolkits/FConsoleCommandExecutor.cpp
+++ b/Engine/Source/Editor/UnrealEd/Private/Toolkits/FConsoleCommandExecutor.cpp
@@ -40,24 +40,24 @@ void FConsoleCommandExecutor::GetSuggestedCompletions(const TCHAR* Input, TArray
 {
 	auto OnConsoleVariable = [&Out](const TCHAR *Name, IConsoleObject* CVar)
 	{
-#if (UE_BUILD_SHIPPING || UE_BUILD_TEST)
-		if (CVar->TestFlags(ECVF_Cheat))
+		if (CVar->IsEnabled())
 		{
-			return;
+			Out.Add(FConsoleSuggestion(Name, CVar->GetDetailedHelp().ToString()));
 		}
-#endif // (UE_BUILD_SHIPPING || UE_BUILD_TEST)
-		if (CVar->TestFlags(ECVF_Unregistered))
-		{
-			return;
-		}
-
-		Out.Add(FConsoleSuggestion(Name, CVar->GetHelp()));
 	};
 
-	IConsoleManager::Get().ForEachConsoleObjectThatContains(FConsoleObjectVisitor::CreateLambda(OnConsoleVariable), Input);
+	IConsoleManager& ConsoleManager = IConsoleManager::Get();
+	ConsoleManager.ForEachConsoleObjectThatContains(FConsoleObjectVisitor::CreateLambda(OnConsoleVariable), Input);
 	for (const FString& CommandName : GetDefault<UConsoleSettings>()->GetFilteredManualAutoCompleteCommands(Input))
 	{
-		Out.Add(FConsoleSuggestion(CommandName, FString()));
+		FString HelpString;
+		// Try to find a console object for this entry in order to retrieve a help string if possible :
+		const TCHAR* CommandNamePtr = *CommandName;
+		if (IConsoleObject* CObj = ConsoleManager.FindConsoleObject(*FParse::Token(CommandNamePtr, /*UseEscape = */false), /*bTrackFrequentCalls = */false); CObj && CObj->IsEnabled())
+		{
+			HelpString = CObj->GetDetailedHelp().ToString();
+		}
+		Out.Add(FConsoleSuggestion(CommandName, HelpString));
 	}
 }
 
diff --git a/Engine/Source/Runtime/Core/Private/HAL/ConsoleManager.cpp b/Engine/Source/Runtime/Core/Private/HAL/ConsoleManager.cpp
--- a/Engine/Source/Runtime/Core/Private/HAL/ConsoleManager.cpp
+++ b/Engine/Source/Runtime/Core/Private/HAL/ConsoleManager.cpp
@@ -348,6 +348,7 @@ public:
 	{
 		return *Help;
 	}
+	
 	virtual void SetHelp(const TCHAR* Value) override final
 	{
 		check(Value);
@@ -357,10 +358,19 @@ public:
 		// for now disabled as there is no good callstack when we crash early during engine init
 //		ensure(IsGoodHelpString(Value));
 	}
+
+	virtual FText GetDetailedHelp() const override
+	{
+		// Append the current value of the variable to provide more details :
+		const FString CurrentValueString = GetString();
+		return FText::Format(LOCTEXT("CVarDetailedHelp", "{0}\n\nCurrent value : {1}"), FText::FromString(GetHelp()), CurrentValueString.IsEmpty() ? LOCTEXT("CVarEmptyValueString", "<empty>") : FText::FromString(CurrentValueString));
+	}
+
 	virtual EConsoleVariableFlags GetFlags() const
 	{
 		return Flags;
 	}
+	
 	virtual void SetFlags(const EConsoleVariableFlags Value)
 	{
 		Flags = Value;
@@ -541,6 +551,7 @@ public:
 	{
 		return *Help;
 	}
+
 	virtual void SetHelp(const TCHAR* InValue)
 	{
 		check(InValue);
@@ -548,10 +559,12 @@ public:
 
 		Help = InValue;
 	}
+
 	virtual EConsoleVariableFlags GetFlags() const
 	{
 		return Flags;
 	}
+
 	virtual void SetFlags(const EConsoleVariableFlags Value)
 	{
 		Flags = Value;
@@ -2745,14 +2758,7 @@ bool FConsoleManager::ProcessUserConsoleInput(const TCHAR* InInput, FOutputDevic
 		return false;
 	}
 
-#if DISABLE_CHEAT_CVARS
-	if(CObj->TestFlags(ECVF_Cheat))
-	{
-		return false;
-	}
-#endif // DISABLE_CHEAT_CVARS
-
-	if(CObj->TestFlags(ECVF_Unregistered))
+	if (!CObj->IsEnabled())
 	{
 		return false;
 	}
@@ -4702,6 +4708,15 @@ public:
 		return TEXT("");
 	}
 
+	virtual FText GetDetailedHelp() const override
+	{
+		if (Bind())
+		{
+			return RealVariable->GetDetailedHelp();
+		}
+		return INVTEXT("");
+	}
+
 	virtual void SetHelp(const TCHAR* Value) override
 	{
 		if (Bind())
@@ -4914,6 +4929,16 @@ struct FConsoleCommandShadow : public FConsoleObjectShadowData, public IConsoleC
 			return RealCommand->SetHelp(InValue);
 		}
 	}
+
+	virtual FText GetDetailedHelp() const
+	{
+		if (Bind())
+		{
+			return RealCommand->GetDetailedHelp();
+		}
+		return INVTEXT("");
+	}
+
 	virtual EConsoleVariableFlags GetFlags() const override
 	{
 		if (Bind())
@@ -4923,6 +4948,7 @@ struct FConsoleCommandShadow : public FConsoleObjectShadowData, public IConsoleC
 		return (EConsoleVariableFlags)0;
 
 	}
+
 	virtual void SetFlags(const EConsoleVariableFlags Value) override
 	{
 		if (BindForUsage())
diff --git a/Engine/Source/Runtime/Core/Public/HAL/IConsoleManager.h b/Engine/Source/Runtime/Core/Public/HAL/IConsoleManager.h
--- a/Engine/Source/Runtime/Core/Public/HAL/IConsoleManager.h
+++ b/Engine/Source/Runtime/Core/Public/HAL/IConsoleManager.h
@@ -12,6 +12,7 @@
 #include "Delegates/Delegate.h"
 #include "Features/IModularFeature.h"
 #include "Templates/EnableIf.h"
+#include "Internationalization/Text.h"
 
 #define TRACK_CONSOLE_FIND_COUNT !(UE_BUILD_SHIPPING || UE_BUILD_TEST)
 
@@ -389,6 +390,14 @@ public:
 	 * Sets the internal flag state to the specified value.
 	 */
 	virtual void SetFlags(const EConsoleVariableFlags Value) = 0;
+	/** 
+	 *  @return a (potentially) more detailed help than GetHelp (e.g. current value for console variables)
+	 */
+	virtual FText GetDetailedHelp() const
+	{
+		// By default, just return the standard help message :
+		return FText::FromString(GetHelp());
+	}
 
 	// Convenience methods -------------------------------------
 
@@ -409,6 +418,25 @@ public:
 		return ((uint32)GetFlags() & (uint32)Value) != 0;
 	}
 
+	/**
+	 * Test the validity of the variable wrt its flags and current build type (e.g. ECVF_Cheat variables are disabled on some targets)
+	 */
+	bool IsEnabled() const
+	{
+#if DISABLE_CHEAT_CVARS
+		if (TestFlags(ECVF_Cheat))
+		{
+			return false;
+		}
+#endif // DISABLE_CHEAT_CVARS
+		if (TestFlags(ECVF_Unregistered))
+		{
+			return false;
+		}
+		return true;
+	}
+
+
 	/**
 	 * If the object has a parent (for instance the main cvar that owns an other-platform cvar), return it
 	 */
@@ -1779,6 +1807,11 @@ public:
 		return TEXT("NO_CVARS, no help");
 	}
 
+	virtual FText GetDetailedHelp() const override
+	{
+		return INVTEXT("NO_CVARS, no help");
+	}
+
 	virtual void SetHelp(const TCHAR* InHelp) override
 	{
 		check(false);
diff --git a/Engine/Source/Runtime/Engine/Private/Engine/EngineConsoleCommandExecutor.cpp b/Engine/Source/Runtime/Engine/Private/Engine/EngineConsoleCommandExecutor.cpp
--- a/Engine/Source/Runtime/Engine/Private/Engine/EngineConsoleCommandExecutor.cpp
+++ b/Engine/Source/Runtime/Engine/Private/Engine/EngineConsoleCommandExecutor.cpp
@@ -52,23 +52,24 @@ void FEngineConsoleCommandExecutor::GetSuggestedCompletions(const TCHAR* Input,
 {
 	const auto OnConsoleVariable = [&Out](const TCHAR* Name, IConsoleObject* CVar)
 	{
-	#if UE_BUILD_SHIPPING || UE_BUILD_TEST
-		if (CVar->TestFlags(ECVF_Cheat))
+		if (CVar->IsEnabled())
 		{
-			return;
+			Out.Add(FConsoleSuggestion(Name, CVar->GetDetailedHelp().ToString()));
 		}
-	#endif // UE_BUILD_SHIPPING || UE_BUILD_TEST
-		if (CVar->TestFlags(ECVF_Unregistered))
-		{
-			return;
-		}
-		Out.Add(FConsoleSuggestion(Name, CVar->GetHelp()));
 	};
 
-	IConsoleManager::Get().ForEachConsoleObjectThatContains(FConsoleObjectVisitor::CreateLambda(OnConsoleVariable), Input);
+	IConsoleManager& ConsoleManager = IConsoleManager::Get();
+	ConsoleManager.ForEachConsoleObjectThatContains(FConsoleObjectVisitor::CreateLambda(OnConsoleVariable), Input);
 	for (const FString& CommandName : GetDefault<UConsoleSettings>()->GetFilteredManualAutoCompleteCommands(Input))
 	{
-		Out.Add(FConsoleSuggestion(CommandName, FString()));
+		FString HelpString;
+		// Try to find a console object for this entry in order to retrieve a help string if possible :
+		const TCHAR* CommandNamePtr = *CommandName;
+		if (IConsoleObject* CObj = ConsoleManager.FindConsoleObject(*FParse::Token(CommandNamePtr, /*UseEscape = */false), /*bTrackFrequentCalls = */false); CObj && CObj->IsEnabled())
+		{
+			HelpString = CObj->GetDetailedHelp().ToString();
+		}
+		Out.Add(FConsoleSuggestion(CommandName, HelpString));
 	}
 }
 
diff --git a/Engine/Source/Runtime/Engine/Private/UserInterface/Console.cpp b/Engine/Source/Runtime/Engine/Private/UserInterface/Console.cpp
--- a/Engine/Source/Runtime/Engine/Private/UserInterface/Console.cpp
+++ b/Engine/Source/Runtime/Engine/Private/UserInterface/Console.cpp
@@ -85,13 +85,7 @@ public:
 	// @param CVar must not be 0
 	static void OnConsoleVariable(const TCHAR* Name, IConsoleObject* CVar, TArray<struct FAutoCompleteCommand>* Sink)
 	{
-#if DISABLE_CHEAT_CVARS
-		if (CVar->TestFlags(ECVF_Cheat))
-		{
-			return;
-		}
-#endif // DISABLE_CHEAT_CVARS
-		if (CVar->TestFlags(ECVF_Unregistered))
+		if (!CVar->IsEnabled())
 		{
 			return;
 		}
