Fixed output log opening in the middle of its history list if history contains an empty entry. Use a TOptional to prevent selecting an empty entry when nothing should be selected.

diff --git a/Engine/Source/Developer/OutputLog/Private/SOutputLog.cpp b/Engine/Source/Developer/OutputLog/Private/SOutputLog.cpp
--- a/Engine/Source/Developer/OutputLog/Private/SOutputLog.cpp
+++ b/Engine/Source/Developer/OutputLog/Private/SOutputLog.cpp
@@ -547,7 +547,7 @@ FReply SConsoleInputBox::OnPreviewKeyDown(const FGeometry& MyGeometry, const FKe
 
 void SConsoleInputBox::SetSuggestions(TArray<FConsoleSuggestion>& Elements, FText Highlight)
 {
-	FString SelectionText;
+	TOptional<FString> SelectionText;
 	if (Suggestions.HasSelectedSuggestion())
 	{
 		SelectionText = Suggestions.GetSelectedSuggestion()->Name;
@@ -560,7 +560,7 @@ void SConsoleInputBox::SetSuggestions(TArray<FConsoleSuggestion>& Elements, FTex
 	{
 		Suggestions.SuggestionsList.Add(MakeShared<FConsoleSuggestion>(Elements[i]));
 
-		if (Elements[i].Name == SelectionText)
+		if (SelectionText.IsSet() && Elements[i].Name == SelectionText.GetValue())
 		{
 			Suggestions.SelectedSuggestion = i;
 		}
@@ -2511,3 +2511,5 @@ void FOutputLogFilter::ClearSelectedLogCategories()
 	SelectedLogCategories.Reset();
 }
 #undef LOCTEXT_NAMESPACE
+
+#pragma optimize("", on)
