OutputLog: Updated the console input widget at the base of the Output Log window to fill the available space.

diff --git a/Engine/Source/Developer/OutputLog/Private/SOutputLog.cpp b/Engine/Source/Developer/OutputLog/Private/SOutputLog.cpp
--- a/Engine/Source/Developer/OutputLog/Private/SOutputLog.cpp
+++ b/Engine/Source/Developer/OutputLog/Private/SOutputLog.cpp
@@ -214,7 +214,7 @@ void SConsoleInputBox::Construct(const FArguments& InArgs)
 			]
 
 			+SHorizontalBox::Slot()
-			.AutoWidth()
+			.FillWidth(1.0f)
 			[
 				SNew(SBox)
 				.MinDesiredWidth(300.f)
