Update OutputLog with improved behavior, display, and saved filter settings.
- Control over verbosity display is much improved.
- The "Show All Categories" checkbox correctly reflects whether or not all categories are displayed.
- Pressing "Dock in Layout" from the Output Log Drawer duplicates its filter state to the docked tab.
- The filter state of the docked tab is saved between editor sessions.

diff --git a/Engine/Source/Developer/OutputLog/Private/OutputLogModule.cpp b/Engine/Source/Developer/OutputLog/Private/OutputLogModule.cpp
--- a/Engine/Source/Developer/OutputLog/Private/OutputLogModule.cpp
+++ b/Engine/Source/Developer/OutputLog/Private/OutputLogModule.cpp
@@ -112,30 +112,31 @@ private:
 TSharedRef<SDockTab> FOutputLogModule::SpawnOutputLogTab(const FSpawnTabArgs& Args)
 {
 	TSharedRef<SOutputLog> NewLog = SNew(SOutputLog, false).Messages(OutputLogHistory->GetMessages());
-	NewLog->UpdateOutputLogFilter(*OutputLogFilterCache);
+	NewLog->UpdateOutputLogFilter(GetDefault<UOutputLogSettings>()->OutputLogTabFilter);
 
 	OutputLog = NewLog;
 
 	TSharedRef<SDockTab> NewTab = SNew(SDockTab)
 		.TabRole(ETabRole::NomadTab)
 		.Label(NSLOCTEXT("OutputLog", "TabTitle", "Output Log"))
+		.OnPersistVisualState_Raw(this, &FOutputLogModule::SaveDockedTabSettings)
 		[
 			NewLog
 		];
 
-	NewTab->SetOnTabClosed(SDockTab::FOnTabClosedCallback::CreateRaw(this, &FOutputLogModule::OnOutputLogTabClosed));
-
 	OutputLogTab = NewTab;
 
 	return NewTab;
 }
 
-void FOutputLogModule::OnOutputLogTabClosed(TSharedRef<SDockTab> Tab)
+void FOutputLogModule::SaveDockedTabSettings()
 {
 	if (TSharedPtr<SOutputLog> SharedOutputLog = OutputLog.Pin())
 	{
 		// Cache the closing LogFilterTab so that we can restore the same filter when it's opened again
-		*OutputLogFilterCache = SharedOutputLog->GetOutputLogFilter();
+		UOutputLogSettings* Settings = GetMutableDefault<UOutputLogSettings>();
+		SharedOutputLog->GetOutputLogFilter().ExportSettings(Settings->OutputLogTabFilter);
+		Settings->SaveConfig();
 	}
 }
 
@@ -458,29 +459,55 @@ void FOutputLogModule::UpdateOutputLogFilter(const TArray<FName>& CategoriesToSh
 {
 	if (InParams.bShowErrors.IsSet())
 	{
-		OutputLogFilterCache->bShowErrors = InParams.bShowErrors.GetValue();
+		OutputLogFilterCache->ErrorsFilter = InParams.bShowErrors.GetValue() ? ELogLevelFilter::Enabled : ELogLevelFilter::None;
 	}
 	// Update the filter cache to these new settings
 	// This will be useful for the case where the OutputLog Drawer or Tab get created after this call
 	if (InParams.bShowWarnings.IsSet())
 	{
-		OutputLogFilterCache->bShowWarnings = InParams.bShowWarnings.GetValue();
+		OutputLogFilterCache->WarningsFilter = InParams.bShowWarnings.GetValue() ? ELogLevelFilter::Enabled : ELogLevelFilter::None;
 	}
 	if (InParams.bShowLogs.IsSet())
 	{
-		OutputLogFilterCache->bShowLogs = InParams.bShowLogs.GetValue();
+		OutputLogFilterCache->MessagesFilter = InParams.bShowLogs.GetValue() ? ELogLevelFilter::Enabled : ELogLevelFilter::None;
 	}
 	if (InParams.IgnoreFilterVerbosities.IsSet())
 	{
-		OutputLogFilterCache->IgnoreFilterVerbosities = InParams.IgnoreFilterVerbosities.GetValue();
+		// Redirect Log/Warning/Error versions to the enum variants.
+		// Only using the set for compatability
+		OutputLogFilterCache->IgnoreFilterVerbosities.Reset();
+		for (ELogVerbosity::Type Verbosity : InParams.IgnoreFilterVerbosities.GetValue())
+		{
+			switch (Verbosity)
+			{
+			case ELogVerbosity::Log:
+				OutputLogFilterCache->MessagesFilter = ELogLevelFilter::All;
+				break;
+			case ELogVerbosity::Warning:
+				OutputLogFilterCache->WarningsFilter = ELogLevelFilter::All;
+				break;
+			case ELogVerbosity::Error:
+				OutputLogFilterCache->ErrorsFilter = ELogLevelFilter::All;
+				break;
+			default:
+				OutputLogFilterCache->IgnoreFilterVerbosities.Add(Verbosity);
+			}
+		}
 	}
 
-	OutputLogFilterCache->ClearSelectedLogCategories();
-	for (const FName& CategoryToShow : CategoriesToShow)
+	if (CategoriesToShow.IsEmpty())
 	{
-		OutputLogFilterCache->ToggleLogCategory(CategoryToShow);
+		OutputLogFilterCache->SetAllCategoriesEnabled(true);
 	}
-
+	else
+	{
+		OutputLogFilterCache->SetAllCategoriesEnabled(false);
+		for (const FName& CategoryToShow : CategoriesToShow)
+		{
+			OutputLogFilterCache->SetLogCategoryEnabled(CategoryToShow, true);
+		}
+	}
+	
 	if (TSharedPtr<SOutputLog> SharedOutputLog = OutputLog.Pin())
 	{
 		SharedOutputLog->UpdateOutputLogFilter(*OutputLogFilterCache);
diff --git a/Engine/Source/Developer/OutputLog/Private/SOutputLog.cpp b/Engine/Source/Developer/OutputLog/Private/SOutputLog.cpp
--- a/Engine/Source/Developer/OutputLog/Private/SOutputLog.cpp
+++ b/Engine/Source/Developer/OutputLog/Private/SOutputLog.cpp
@@ -39,6 +39,7 @@
 #include "OutputLogStyle.h"
 #include "OutputLogMenuContext.h"
 #include "ToolMenus.h"
+#include "Widgets/Input/SSegmentedControl.h"
 
 
 #define LOCTEXT_NAMESPACE "SOutputLog"
@@ -1090,12 +1091,6 @@ int32 FOutputLogTextLayoutMarshaller::GetNumMessages() const
 
 int32 FOutputLogTextLayoutMarshaller::GetNumFilteredMessages()
 {
-	// No need to filter the messages if the filter is not set
-	if (!Filter->IsFilterSet())
-	{
-		return GetNumMessages();
-	}
-
 	// Re-count messages if filter changed before we refresh
 	if (bNumMessagesCacheDirty)
 	{
@@ -1771,25 +1766,13 @@ void SOutputLog::BuildInitialLogCategoryFilter(const FArguments& InArgs)
 {
 	for (const auto& Message : InArgs._Messages)
 	{
-		Filter.AddAvailableLogCategory(Message->Category);
 		const bool bIsDeselectedByDefault = InArgs._AllowInitialLogCategory.IsBound() && !InArgs._AllowInitialLogCategory.Execute(Message->Category);
-		if (bIsDeselectedByDefault && Filter.IsLogCategoryEnabled(Message->Category))
-		{
-			Filter.bShowAllCategories = false;
-			Filter.ToggleLogCategory(Message->Category);
-		}
+		Filter.AddAvailableLogCategory(Message->Category, bIsDeselectedByDefault ? false : TOptional<bool>());
 	}
 
 	for (auto DefaultCategorySelectionIt = InArgs._DefaultCategorySelection.CreateConstIterator(); DefaultCategorySelectionIt; ++DefaultCategorySelectionIt)
 	{
-		const FName Category = DefaultCategorySelectionIt->Key;
-		Filter.AddAvailableLogCategory(Category);
-
-		Filter.bShowAllCategories &= DefaultCategorySelectionIt->Value;
-		if (Filter.IsLogCategoryEnabled(Category) != DefaultCategorySelectionIt->Value)
-		{
-			Filter.ToggleLogCategory(Category);
-		}
+		Filter.SetLogCategoryEnabled(DefaultCategorySelectionIt->Key, DefaultCategorySelectionIt->Value);
 	}
 }
 
@@ -1826,58 +1809,122 @@ TSharedRef<SWidget> SOutputLog::MakeAddFilterMenu()
 {
 	FMenuBuilder MenuBuilder(/*bInShouldCloseWindowAfterMenuSelection=*/true, nullptr);
 	
-	MenuBuilder.BeginSection("OutputLogMiscEntries", LOCTEXT("OutputLogFilters", "Filters"));
+	MenuBuilder.BeginSection("OutputLogVerbosityEntries", LOCTEXT("OutputLogVerbosityHeading", "Verbosity"));
 	{
-		MenuBuilder.AddMenuEntry(
-			LOCTEXT("ShowAllCategories", "Show All"),
-			LOCTEXT("ShowAllCategories_Tooltip", "Filter the Output Log to show all categories"),
-			FSlateIcon(),
-			FUIAction(FExecuteAction::CreateSP(this, &SOutputLog::CategoriesShowAll_Execute),
-				FCanExecuteAction::CreateLambda([] { return true; }),
-				FIsActionChecked::CreateSP(this, &SOutputLog::CategoriesShowAll_IsChecked)),
-			NAME_None,
-			EUserInterfaceActionType::ToggleButton
+		const FText AllLabel = LOCTEXT("AllLabel", "All");
+		const FText EnabledLabel = LOCTEXT("EnabledLabel", "Filtered");
+		const FText NoneLabel = LOCTEXT("NoneLabel", "None");
+		
+		MenuBuilder.AddWidget(
+			SNew(SHorizontalBox)
+				+ SHorizontalBox::Slot()
+				.FillWidth(1.0)
+				.VAlign(VAlign_Center)
+				.Padding(12, 0, 4, 0)
+				[
+					SNew(STextBlock)
+					.Text(LOCTEXT("Messages", "Messages"))
+					.TextStyle(FAppStyle::Get(), "Menu.Label")
+					.ColorAndOpacity(FLinearColor::White)
+				]
+				+ SHorizontalBox::Slot()
+				.AutoWidth()
+				.Padding(4, 0, 12, 2)
+				[
+					SNew(SSegmentedControl<ELogLevelFilter>)
+					.Value(this, &SOutputLog::GetMessagesFilter)
+					.OnValueChanged(this, &SOutputLog::OnMessagesFilterChanged)
+					+ SSegmentedControl<ELogLevelFilter>::Slot(ELogLevelFilter::None)
+					.Text(NoneLabel)
+					.ToolTip(LOCTEXT("NoMessagesTooltip", "No messages will be shown."))
+					+ SSegmentedControl<ELogLevelFilter>::Slot(ELogLevelFilter::Enabled)
+					.Text(EnabledLabel)
+					.ToolTip(LOCTEXT("EnabledMessagesTooltip", "Show messages from the enabled categories."))
+					+ SSegmentedControl<ELogLevelFilter>::Slot(ELogLevelFilter::All)
+					.Text(AllLabel)
+					.ToolTip(LOCTEXT("AllMessagesTooltip", "Show all messages, ignoring whether or not the category is enabled."))
+				],
+			FText::GetEmpty(),
+			true
 		);
 
-		MenuBuilder.AddSubMenu(
-			LOCTEXT("Filters", "Filters"),
-			LOCTEXT("SelectCategoriesToolTip", "Set filter settings such as visible categories"),
-			FNewMenuDelegate::CreateSP(this, &SOutputLog::MakeSelectCategoriesSubMenu)
+		MenuBuilder.AddWidget(
+			SNew(SHorizontalBox)
+				+ SHorizontalBox::Slot()
+				.FillWidth(1.0)
+				.VAlign(VAlign_Center)
+				.Padding(12, 0, 4, 0)
+				[
+					SNew(STextBlock)
+					.Text(LOCTEXT("Warnings", "Warnings"))
+					.TextStyle(FAppStyle::Get(), "Menu.Label")
+					.ColorAndOpacity(FLinearColor::White)
+				]
+				+ SHorizontalBox::Slot()
+				.AutoWidth()
+				.Padding(4, 0, 12, 2)
+				[
+					SNew(SSegmentedControl<ELogLevelFilter>)
+					.Value(this, &SOutputLog::GetWarningsFilter)
+					.OnValueChanged(this, &SOutputLog::OnWarningsFilterChanged)
+					+ SSegmentedControl<ELogLevelFilter>::Slot(ELogLevelFilter::None)
+					.Text(NoneLabel)
+					.ToolTip(LOCTEXT("NoWarningsTooltip", "No warnings will be shown."))
+					+ SSegmentedControl<ELogLevelFilter>::Slot(ELogLevelFilter::Enabled)
+					.Text(EnabledLabel)
+					.ToolTip(LOCTEXT("EnabledWarningsTooltip", "Show warnings from the enabled categories."))
+					+ SSegmentedControl<ELogLevelFilter>::Slot(ELogLevelFilter::All)
+					.Text(AllLabel)
+					.ToolTip(LOCTEXT("AllWarningsTooltip", "Show all warnings, ignoring whether or not the category is enabled."))
+				],
+			FText::GetEmpty(),
+			true
+		);
+
+		MenuBuilder.AddWidget(
+			SNew(SHorizontalBox)
+				+ SHorizontalBox::Slot()
+				.FillWidth(1.0)
+				.VAlign(VAlign_Center)
+				.Padding(12, 0, 4, 0)
+				[
+					SNew(STextBlock)
+					.Text(LOCTEXT("Errors", "Errors"))
+					.TextStyle(FAppStyle::Get(), "Menu.Label")
+					.ColorAndOpacity(FLinearColor::White)
+				]
+				+ SHorizontalBox::Slot()
+				.AutoWidth()
+				.Padding(4, 0, 12, 2)
+				[
+					SNew(SSegmentedControl<ELogLevelFilter>)
+					.Value(this, &SOutputLog::GetErrorsFilter)
+					.OnValueChanged(this, &SOutputLog::OnErrorsFilterChanged)
+					+ SSegmentedControl<ELogLevelFilter>::Slot(ELogLevelFilter::None)
+					.Text(NoneLabel)
+					.ToolTip(LOCTEXT("NoErrorsTooltip", "No errors will be shown."))
+					+ SSegmentedControl<ELogLevelFilter>::Slot(ELogLevelFilter::Enabled)
+					.Text(EnabledLabel)
+					.ToolTip(LOCTEXT("EnabledErrorsTooltip", "Show errors from the enabled categories."))
+					+ SSegmentedControl<ELogLevelFilter>::Slot(ELogLevelFilter::All)
+					.Text(AllLabel)
+					.ToolTip(LOCTEXT("AllErrorsTooltip", "Show all errors, ignoring whether or not the category is enabled."))
+				],
+			FText::GetEmpty(),
+			true
 		);
 	}
 	MenuBuilder.EndSection();
 
-	MenuBuilder.BeginSection("OutputLogVerbosityEntries", LOCTEXT("OutputLogVerbosityHeading", "Verbosity"));
+	MenuBuilder.BeginSection("OutputLogBetterFilter", LOCTEXT("OutputLogFilterCategories", "Categories"));
 	{
-		MenuBuilder.AddMenuEntry(
-			LOCTEXT("ShowMessages", "Messages"), 
-			LOCTEXT("ShowMessages_Tooltip", "[Checked] Filter Output Log to show messages within selected categories. [-] Filter Output Log to show all messages"), 
-			FSlateIcon(), 
-			FUIAction(FExecuteAction::CreateSP(this, &SOutputLog::VerbosityLogs_Execute), 
-				FCanExecuteAction::CreateLambda([] { return true; }), 
-				FGetActionCheckState::CreateSP(this, &SOutputLog::VerbosityLogs_IsChecked)),
-			NAME_None, 
-			EUserInterfaceActionType::ToggleButton
-		);
-
-		MenuBuilder.AddMenuEntry(
-			LOCTEXT("ShowWarnings", "Warnings"),
-			LOCTEXT("ShowWarnings_Tooltip", "[Checked] Filter Output Log to show warnings within selected categories. [-] Filter Output Log to show all warnings"),
-			FSlateIcon(),
-			FUIAction(FExecuteAction::CreateSP(this, &SOutputLog::VerbosityWarnings_Execute),
-				FCanExecuteAction::CreateLambda([] { return true; }),
-				FGetActionCheckState::CreateSP(this, &SOutputLog::VerbosityWarnings_IsChecked)),
-			NAME_None,
-			EUserInterfaceActionType::ToggleButton
-		);
-
-		MenuBuilder.AddMenuEntry(
-			LOCTEXT("ShowErrors", "Errors"),
-			LOCTEXT("ShowErrors_Tooltip", "[Checked] Filter Output Log to show errors within selected categories. [-] Filter Output Log to show all errors"),
-			FSlateIcon(),
-			FUIAction(FExecuteAction::CreateSP(this, &SOutputLog::VerbosityErrors_Execute),
+		MenuBuilder.AddSubMenu(
+			LOCTEXT("AllFilterCategories", "Category Filters"),
+			LOCTEXT("AllFilterCategoriesTooltip", "Select the log categories that are displayed."),
+			FNewMenuDelegate::CreateSP(this, &SOutputLog::MakeSelectCategoriesSubMenu),
+			FUIAction(FExecuteAction::CreateSP(this, &SOutputLog::CategoriesShowAll_Execute),
 				FCanExecuteAction::CreateLambda([] { return true; }),
-				FGetActionCheckState::CreateSP(this, &SOutputLog::VerbosityErrors_IsChecked)),
+				FGetActionCheckState::CreateSP(this, &SOutputLog::CategoriesShowAll_IsChecked)),
 			NAME_None,
 			EUserInterfaceActionType::ToggleButton
 		);
@@ -1891,15 +1938,17 @@ void SOutputLog::MakeSelectCategoriesSubMenu(FMenuBuilder& MenuBuilder)
 {
 	MenuBuilder.BeginSection("OutputLogCategoriesEntries");
 	{	
-		for (const FName& Category : Filter.GetAvailableLogCategories())
+		for (const FOutputLogCategorySettings& Category : Filter.GetCategoryFilters())
 		{
+			FString NameString = Category.Name.ToString();
+			
 			MenuBuilder.AddMenuEntry(
-				FText::AsCultureInvariant(Category.ToString()),
-				FText::Format(LOCTEXT("Category_Tooltip", "Filter the Output Log to show category: {0}"), FText::AsCultureInvariant(Category.ToString())),
+				FText::AsCultureInvariant(NameString),
+				FText::Format(LOCTEXT("Category_Tooltip", "Filter the Output Log to show category: {0}"), FText::AsCultureInvariant(NameString)),
 				FSlateIcon(),
-				FUIAction(FExecuteAction::CreateSP(this, &SOutputLog::CategoriesSingle_Execute, Category),
+				FUIAction(FExecuteAction::CreateSP(this, &SOutputLog::CategoriesSingle_Execute, Category.Name),
 				FCanExecuteAction::CreateLambda([] { return true; }),
-				FIsActionChecked::CreateSP(this, &SOutputLog::CategoriesSingle_IsChecked, Category)),
+				FIsActionChecked::CreateSP(this, &SOutputLog::CategoriesSingle_IsChecked, Category.Name)),
 				NAME_None,
 				EUserInterfaceActionType::ToggleButton
 			);
@@ -1908,112 +1957,45 @@ void SOutputLog::MakeSelectCategoriesSubMenu(FMenuBuilder& MenuBuilder)
 	MenuBuilder.EndSection();
 }
 
-ECheckBoxState SOutputLog::VerbosityLogs_IsChecked() const
+ELogLevelFilter SOutputLog::GetMessagesFilter() const
 {
-	if (Filter.bShowLogs && Filter.IgnoreFilterVerbosities.Contains(ELogVerbosity::Log))
-	{
-		return ECheckBoxState::Undetermined;
-	}
-
-	if (Filter.bShowLogs)
-	{
-		return ECheckBoxState::Checked;
-	}
-
-	return ECheckBoxState::Unchecked;
+	return Filter.MessagesFilter;
 }
 
-ECheckBoxState SOutputLog::VerbosityWarnings_IsChecked() const
+void SOutputLog::OnMessagesFilterChanged(ELogLevelFilter NewFilter)
 {
-	if (Filter.bShowWarnings && Filter.IgnoreFilterVerbosities.Contains(ELogVerbosity::Warning))
-	{
-		return ECheckBoxState::Undetermined;
-	}
-
-	if (Filter.bShowWarnings)
-	{
-		return ECheckBoxState::Checked;
-	}
-
-	return ECheckBoxState::Unchecked;
+	Filter.MessagesFilter = NewFilter;
+	MessagesTextMarshaller->MarkMessagesCacheAsDirty();
+	Refresh();
 }
 
-ECheckBoxState SOutputLog::VerbosityErrors_IsChecked() const
+ELogLevelFilter SOutputLog::GetWarningsFilter() const
 {
-	if (Filter.bShowErrors && Filter.IgnoreFilterVerbosities.Contains(ELogVerbosity::Error))
-	{
-		return ECheckBoxState::Undetermined;
-	}
-
-	if (Filter.bShowErrors)
-	{
-		return ECheckBoxState::Checked;
-	}
-
-	return ECheckBoxState::Unchecked;
+	return Filter.WarningsFilter;
 }
 
-static void VerbosityGeneric_Execute(bool& Flag, TSet<ELogVerbosity::Type>& Filter, ELogVerbosity::Type Verbosity)
+void SOutputLog::OnWarningsFilterChanged(ELogLevelFilter NewFilter)
 {
-	if (FSlateApplication::Get().GetModifierKeys().IsShiftDown())
-	{
-		Flag = !Flag;
-	}
-	else
-	{
-		// Rotate through: showing the verbosity, showing the verbosity while ignoring filter categories, and hiding the verbosity
-		if (Flag)
-		{
-			if (Filter.Contains(Verbosity))
-			{
-				Flag = false;
-				Filter.Remove(Verbosity);
-			}
-			else
-			{
-				Filter.Emplace(Verbosity);
-			}
-		}
-		else
-		{
-			Flag = true;
-		}
-	}
-}
-
-void SOutputLog::VerbosityLogs_Execute()
-{ 
-	VerbosityGeneric_Execute(Filter.bShowLogs, Filter.IgnoreFilterVerbosities, ELogVerbosity::Log);
-
-	// Flag the messages count as dirty
+	Filter.WarningsFilter = NewFilter;
 	MessagesTextMarshaller->MarkMessagesCacheAsDirty();
-
 	Refresh();
 }
 
-void SOutputLog::VerbosityWarnings_Execute()
+ELogLevelFilter SOutputLog::GetErrorsFilter() const
 {
-	VerbosityGeneric_Execute(Filter.bShowWarnings, Filter.IgnoreFilterVerbosities, ELogVerbosity::Warning);
-
-	// Flag the messages count as dirty
-	MessagesTextMarshaller->MarkMessagesCacheAsDirty();
-
-	Refresh();
+	return Filter.ErrorsFilter;
 }
 
-void SOutputLog::VerbosityErrors_Execute()
+void SOutputLog::OnErrorsFilterChanged(ELogLevelFilter NewFilter)
 {
-	VerbosityGeneric_Execute(Filter.bShowErrors, Filter.IgnoreFilterVerbosities, ELogVerbosity::Error);
-
-	// Flag the messages count as dirty
+	Filter.ErrorsFilter = NewFilter;
 	MessagesTextMarshaller->MarkMessagesCacheAsDirty();
-
 	Refresh();
 }
 
-bool SOutputLog::CategoriesShowAll_IsChecked() const
+ECheckBoxState SOutputLog::CategoriesShowAll_IsChecked() const
 {
-	return Filter.bShowAllCategories;
+	return Filter.AreAllCategoriesSelected();
 }
 
 bool SOutputLog::CategoriesSingle_IsChecked(FName InName) const
@@ -2023,17 +2005,11 @@ bool SOutputLog::CategoriesSingle_IsChecked(FName InName) const
 
 void SOutputLog::CategoriesShowAll_Execute()
 {
-	Filter.bShowAllCategories = !Filter.bShowAllCategories;
-
-	Filter.ClearSelectedLogCategories();
-	if (Filter.bShowAllCategories)
-	{
-		for (const auto& AvailableCategory : Filter.GetAvailableLogCategories())
-		{
-			Filter.ToggleLogCategory(AvailableCategory);
-		}
-	}
+	const ECheckBoxState CurrentState = Filter.AreAllCategoriesSelected();
+	const bool NextState = CurrentState != ECheckBoxState::Checked;
 
+	Filter.SetAllCategoriesEnabled(NextState);
+	
 	// Flag the messages count as dirty
 	MessagesTextMarshaller->MarkMessagesCacheAsDirty();
 
@@ -2050,39 +2026,17 @@ void SOutputLog::CategoriesSingle_Execute(FName InName)
 	Refresh();
 }
 
-void SOutputLog::UpdateOutputLogFilter(const TArray<FName>& CategoriesToShow, TOptional<bool> bShowErrors, TOptional<bool> bShowWarnings, TOptional<bool> bShowLogs)
+void SOutputLog::UpdateOutputLogFilter(const FOutputLogFilter& InFilter)
 {
-	if (bShowErrors.IsSet())
-	{
-		Filter.bShowErrors = bShowErrors.GetValue();
-	}
-	if (bShowWarnings.IsSet())
-	{
-		Filter.bShowWarnings = bShowWarnings.GetValue();
-	}
-	if (bShowLogs.IsSet())
-	{
-		Filter.bShowLogs = bShowLogs.GetValue();
-	}
-
-	// Show all categories if empty list is passed in. This means there's no way to change the Show* bools
-	// and also leave the categories as-is, but that use case is probably minimal. It can be added if needed.
-	Filter.bShowAllCategories = CategoriesToShow.Num() == 0;
-	Filter.ClearSelectedLogCategories();
-	const TArray<FName> ActualCategoriesToShow = Filter.bShowAllCategories ? Filter.GetAvailableLogCategories() : CategoriesToShow;
-	for (const auto& AvailableCategory : ActualCategoriesToShow)
-	{
-		Filter.ToggleLogCategory(AvailableCategory);
-	}
-
+	Filter = InFilter;
 	MessagesTextMarshaller->MarkMessagesCacheAsDirty();
 	Refresh();
 }
 
-void SOutputLog::UpdateOutputLogFilter(const FOutputLogFilter& InFilter)
+void SOutputLog::UpdateOutputLogFilter(const FOutputLogFilterSettings& InSettings)
 {
-	Filter = InFilter;
-	Filter.bShowAllCategories = Filter.GetSelectedLogCategories().Num() == 0 ;
+	Filter.ApplySettings(InSettings);
+	
 	MessagesTextMarshaller->MarkMessagesCacheAsDirty();
 	Refresh();
 }
@@ -2403,6 +2357,11 @@ FReply SOutputLog::OnDockInLayoutClicked()
 {
 	TSharedPtr<SDockTab> DockedTab;
 
+	// Export our settings so that the docked tab starts from the current state.
+	UOutputLogSettings* Settings = GetMutableDefault<UOutputLogSettings>();
+	Filter.ExportSettings(Settings->OutputLogTabFilter);
+	Settings->SaveConfig();
+	
 	static const FName OutputLogTabName = FName("OutputLog");
 	if (TSharedPtr<SDockTab> ActiveTab = FGlobalTabmanager::Get()->GetActiveTab())
 	{
@@ -2420,95 +2379,205 @@ FReply SOutputLog::OnDockInLayoutClicked()
 	return FReply::Handled();
 }
 
-
-bool FOutputLogFilter::IsMessageAllowed(const TSharedPtr<FOutputLogMessage>& Message)
+ECheckBoxState FOutputLogFilter::AreAllCategoriesSelected() const
 {
-	// Filter Verbosity
+	int32 Count = 0;
+	for (const FOutputLogCategorySettings& Item : Categories)
 	{
-		if (Message->Verbosity == ELogVerbosity::Error && !bShowErrors)
-		{
-			return false;
-		}
-		else if (Message->Verbosity == ELogVerbosity::Warning && !bShowWarnings)
+		if (Item.bEnabled)
 		{
-			return false;
+			Count++; 
 		}
-		else if (Message->Verbosity != ELogVerbosity::Error && Message->Verbosity != ELogVerbosity::Warning && !bShowLogs)
+	}
+	
+	if (Count == Categories.Num())
+	{
+		return ECheckBoxState::Checked;
+	}
+	if (Count == 0)
+	{
+		return ECheckBoxState::Unchecked;
+	}
+	return ECheckBoxState::Undetermined;
+}
+
+void FOutputLogFilter::ApplySettings(const FOutputLogFilterSettings& InSettings)
+{
+	MessagesFilter = InSettings.MessagesFilter;
+	WarningsFilter = InSettings.WarningsFilter;
+	ErrorsFilter = InSettings.ErrorsFilter;
+
+	bSelectNewCategories = InSettings.bSelectNewCategories;
+	
+	SetFilterText(InSettings.FilterText);
+
+	if (InSettings.Categories.Num() == 0)
+	{
+		// This implies all _should_ be selected
+		SetAllCategoriesEnabled(true);
+	}
+	else
+	{
+		// Clear and apply piecemeal so that configured logs that 
+		// haven't been hit yet are still added.
+		SetAllCategoriesEnabled(false);
+		for (const FOutputLogCategorySettings& Category : InSettings.Categories)
 		{
-			return false;
+			SetLogCategoryEnabled(Category.Name, Category.bEnabled);
 		}
 	}
+}
+
+void FOutputLogFilter::ExportSettings(FOutputLogFilterSettings& OutSettings) const
+{
+	OutSettings.MessagesFilter = MessagesFilter;
+	OutSettings.WarningsFilter = WarningsFilter;
+	OutSettings.ErrorsFilter = ErrorsFilter;
+	
+	OutSettings.bSelectNewCategories = bSelectNewCategories;
+
+	if (AreAllCategoriesSelected() == ECheckBoxState::Checked)
+	{
+		OutSettings.Categories.Empty();
+	}
+	else
+	{
+		OutSettings.Categories = Categories;
+	}
+}
+
+bool FOutputLogFilter::IsMessageAllowed(const TSharedPtr<FOutputLogMessage>& Message) const
+{
+	// Filter Verbosity
+	const ELogLevelFilter LevelFilter = GetMessageLevelFilter(Message);
+	if (LevelFilter == ELogLevelFilter::None)
+	{
+		return false;
+	}
 
 	// Filter by Category
+	if (!IgnoreFilterVerbosities.Contains(Message->Verbosity) && LevelFilter == ELogLevelFilter::Enabled && !IsLogCategoryEnabled(Message->Category))
 	{
-		if (!IgnoreFilterVerbosities.Contains(Message->Verbosity) && !IsLogCategoryEnabled(Message->Category))
-		{
-			return false;
-		}
+		return false;
 	}
 
 	// Filter search phrase
+	if (!TextFilterExpressionEvaluator.TestTextFilter(FLogFilter_TextFilterExpressionContextOutputLog(*Message)))
 	{
-		if (!TextFilterExpressionEvaluator.TestTextFilter(FLogFilter_TextFilterExpressionContextOutputLog(*Message)))
-		{
-			return false;
-		}
+		return false;
 	}
 
 	return true;
 }
 
-void FOutputLogFilter::AddAvailableLogCategory(const FName& LogCategory)
+ELogLevelFilter FOutputLogFilter::GetMessageLevelFilter(const TSharedPtr<FOutputLogMessage>& Message) const
+{
+	switch (Message->Verbosity)
+	{
+	case ELogVerbosity::Error:
+		return ErrorsFilter;
+	case ELogVerbosity::Warning:
+		return WarningsFilter;
+	default:
+		return MessagesFilter;
+	}
+}
+
+void FOutputLogFilter::AddAvailableLogCategory(const FName& LogCategory, TOptional<bool> InitiallySelected)
 {
 	// Use an insert-sort to keep AvailableLogCategories alphabetically sorted
 	int32 InsertIndex = 0;
-	for (InsertIndex = AvailableLogCategories.Num()-1; InsertIndex >= 0; --InsertIndex)
+	for (InsertIndex = Categories.Num()-1; InsertIndex >= 0; --InsertIndex)
 	{
-		FName CheckCategory = AvailableLogCategories[InsertIndex];
+		FOutputLogCategorySettings& CheckCategory = Categories[InsertIndex];
 		// No duplicates
-		if (CheckCategory == LogCategory)
+		if (CheckCategory.Name == LogCategory)
 		{
 			return;
 		}
-		else if (CheckCategory.Compare(LogCategory) < 0)
+		else if (CheckCategory.Name.Compare(LogCategory) < 0)
 		{
 			break;
 		}
 	}
-	AvailableLogCategories.Insert(LogCategory, InsertIndex+1);
-	if (bShowAllCategories)
+	Categories.Insert({ LogCategory, InitiallySelected.Get(bSelectNewCategories) }, InsertIndex+1);
+}
+
+bool FOutputLogFilter::IsLogCategoryAvailable(const FName& LogCategory) const
+{
+	return FindCategoryFilter(LogCategory) != nullptr;
+}
+
+FOutputLogCategorySettings* FOutputLogFilter::FindCategoryFilter(const FName& LogCategory)
+{
+	// TODO: Could in theory use a binary search as AddAvailableLogCategory() ensures `Categories` will be sorted.
+	for (FOutputLogCategorySettings& Item : Categories)
+	{
+		if (Item.Name == LogCategory)
+		{
+			return &Item;
+		}
+	}
+	return nullptr;
+}
+
+const FOutputLogCategorySettings* FOutputLogFilter::FindCategoryFilter(const FName& LogCategory) const
+{
+	// TODO: Could in theory use a binary search as AddAvailableLogCategory() ensures `Categories` will be sorted.
+	for (const FOutputLogCategorySettings& Item : Categories)
 	{
-		ToggleLogCategory(LogCategory);
+		if (Item.Name == LogCategory)
+		{
+			return &Item;
+		}
 	}
+	return nullptr;
 }
 
-bool FOutputLogFilter::IsLogCategoryAvailable(const FName& LogCategory) const
+void FOutputLogFilter::SetLogCategoryEnabled(const FName& LogCategory, bool bEnabled)
 {
-	return AvailableLogCategories.Contains(LogCategory);
+	if (FOutputLogCategorySettings* CategoryFilter = FindCategoryFilter(LogCategory))
+	{
+		CategoryFilter->bEnabled = bEnabled;
+	}
+	else
+	{
+		AddAvailableLogCategory(LogCategory, bEnabled);
+	}
+}
+
+void FOutputLogFilter::SetAllCategoriesEnabled(bool bEnabled)
+{
+	for (FOutputLogCategorySettings& Item : Categories)
+	{
+		Item.bEnabled = bEnabled;
+	}
 }
 
 void FOutputLogFilter::ToggleLogCategory(const FName& LogCategory)
 {
-	int32 FoundIndex = SelectedLogCategories.Find(LogCategory);
-	if (FoundIndex == INDEX_NONE)
+	if (FOutputLogCategorySettings* CategoryFilter = FindCategoryFilter(LogCategory))
 	{
-		SelectedLogCategories.Add(LogCategory);
+		CategoryFilter->bEnabled = !CategoryFilter->bEnabled;
 	}
 	else
 	{
-		SelectedLogCategories.RemoveAt(FoundIndex, /*Count=*/1, EAllowShrinking::No);
+		AddAvailableLogCategory(LogCategory, true);
 	}
 }
 
 bool FOutputLogFilter::IsLogCategoryEnabled(const FName& LogCategory) const
 {
-	return SelectedLogCategories.Contains(LogCategory);
+	if (const FOutputLogCategorySettings* CategoryFilter = FindCategoryFilter(LogCategory))
+	{
+		return CategoryFilter->bEnabled;
+	}
+	return false;
 }
 
 void FOutputLogFilter::ClearSelectedLogCategories()
 {
-	// No need to churn memory each time the selected categories are cleared
-	SelectedLogCategories.Reset();
+	SetAllCategoriesEnabled(false);
 }
 #undef LOCTEXT_NAMESPACE
 
diff --git a/Engine/Source/Developer/OutputLog/Private/SOutputLog.h b/Engine/Source/Developer/OutputLog/Private/SOutputLog.h
--- a/Engine/Source/Developer/OutputLog/Private/SOutputLog.h
+++ b/Engine/Source/Developer/OutputLog/Private/SOutputLog.h
@@ -7,6 +7,7 @@
 #include "Widgets/DeclarativeSyntaxSupport.h"
 #include "Input/Reply.h"
 #include "OutputLogCreationParams.h"
+#include "OutputLogSettings.h"
 #include "Widgets/SWidget.h"
 #include "Widgets/SCompoundWidget.h"
 #include "Widgets/Input/SMultiLineEditableTextBox.h"
@@ -242,32 +243,28 @@ private:
 */
 struct FOutputLogFilter
 {
-	/** true to show Logs. */
-	bool bShowLogs;
-
-	/** true to show Warnings. */
-	bool bShowWarnings;
-
-	/** true to show Errors. */
-	bool bShowErrors;
-
-	/** true to allow all Log Categories */
-	bool bShowAllCategories;
+	ELogLevelFilter MessagesFilter = ELogLevelFilter::Enabled;
+	ELogLevelFilter WarningsFilter = ELogLevelFilter::Enabled;
+	ELogLevelFilter ErrorsFilter = ELogLevelFilter::Enabled;
 
 	/** Set of Verbosity levels that will show all regardless of category filter */
 	TSet<ELogVerbosity::Type> IgnoreFilterVerbosities;
 
+	/** Whether newly added categories should begin as selected. */
+	bool bSelectNewCategories = true;
+
 	/** Enable all filters by default */
 	FOutputLogFilter() : TextFilterExpressionEvaluator(ETextFilterExpressionEvaluatorMode::BasicString)
 	{
-		bShowErrors = bShowLogs = bShowWarnings = bShowAllCategories = true;
 	}
 
-	/** Returns true if any messages should be filtered out */
-	bool IsFilterSet() { return !bShowErrors || !bShowLogs || !bShowWarnings || TextFilterExpressionEvaluator.GetFilterType() != ETextFilterExpressionType::Empty || !TextFilterExpressionEvaluator.GetFilterText().IsEmpty(); }
+	ECheckBoxState AreAllCategoriesSelected() const;
 
+	void ApplySettings(const FOutputLogFilterSettings& InSettings);
+	void ExportSettings(FOutputLogFilterSettings& OutSettings) const;
+	
 	/** Checks the given message against set filters */
-	bool IsMessageAllowed(const TSharedPtr<FOutputLogMessage>& Message);
+	bool IsMessageAllowed(const TSharedPtr<FOutputLogMessage>& Message) const;
 
 	/** Set the Text to be used as the Filter's restrictions */
 	void SetFilterText(const FText& InFilterText) { TextFilterExpressionEvaluator.SetFilterText(InFilterText); }
@@ -278,34 +275,38 @@ struct FOutputLogFilter
 	/** Returns Evaluator syntax errors (if any) */
 	FText GetSyntaxErrors() { return TextFilterExpressionEvaluator.GetFilterErrorText(); }
 
-	const TArray<FName>& GetAvailableLogCategories() { return AvailableLogCategories; }
-
-	const TArray<FName>& GetSelectedLogCategories() { return SelectedLogCategories; }
-
+	const TArray<FOutputLogCategorySettings>& GetCategoryFilters() const { return Categories; }
+	
 	/** Adds a Log Category to the list of available categories, if it isn't already present */
-	void AddAvailableLogCategory(const FName& LogCategory);
+	void AddAvailableLogCategory(const FName& LogCategory, TOptional<bool> InitiallySelected = TOptional<bool>());
 
 	/** Returns true if the specified log category is in the list of available categories */
 	bool IsLogCategoryAvailable(const FName& LogCategory) const;
+	
+	/** Returns true if the specified log category is enabled */
+	bool IsLogCategoryEnabled(const FName& LogCategory) const;
+
+	/** Enables or disables a log category if it is not already in that state. */
+	void SetLogCategoryEnabled(const FName& LogCategory, bool bEnabled);
 
+	void SetAllCategoriesEnabled(bool bEnabled);
+	
 	/** Enables or disables a Log Category in the filter */
 	void ToggleLogCategory(const FName& LogCategory);
 
-	/** Returns true if the specified log category is enabled */
-	bool IsLogCategoryEnabled(const FName& LogCategory) const;
-
 	/** Empties the list of selected log categories */
 	void ClearSelectedLogCategories();
-
+	
 private:
+	FOutputLogCategorySettings* FindCategoryFilter(const FName& LogCategory);
+	const FOutputLogCategorySettings* FindCategoryFilter(const FName& LogCategory) const;
+	ELogLevelFilter GetMessageLevelFilter(const TSharedPtr<FOutputLogMessage>& Message) const;
+	
 	/** Expression evaluator that can be used to perform complex text filter queries */
 	FTextFilterExpressionEvaluator TextFilterExpressionEvaluator;
 
-	/** Array of Log Categories which are available for filter -- i.e. have been used in a log this session */
-	TArray<FName> AvailableLogCategories;
-
-	/** Array of Log Categories which are being used in the filter */
-	TArray<FName> SelectedLogCategories;
+	/** All categories the filter is aware of */
+	TArray<FOutputLogCategorySettings> Categories;
 };
 
 /**
@@ -394,8 +395,8 @@ public:
 	void FocusConsoleCommandBox();
 
 	/** Change the output log's filter. If CategoriesToShow is empty, all categories will be shown. */
-	void UpdateOutputLogFilter(const TArray<FName>& CategoriesToShow, TOptional<bool> bShowErrors = TOptional<bool>(), TOptional<bool> bShowWarnings = TOptional<bool>(), TOptional<bool> bShowLogs = TOptional<bool>());
 	void UpdateOutputLogFilter(const FOutputLogFilter& InFilter);
+	void UpdateOutputLogFilter(const FOutputLogFilterSettings& InSettings);
 
 	const FOutputLogFilter& GetOutputLogFilter() { return Filter; }
 protected:
@@ -452,29 +453,20 @@ private:
 	/** Make the "Categories" sub-menu. */
 	void MakeSelectCategoriesSubMenu(FMenuBuilder& MenuBuilder);
 
-	/** Toggles Verbosity "Logs" true/false. */
-	void VerbosityLogs_Execute();
-
-	/** Returns the state of Verbosity "Logs". */
-	ECheckBoxState VerbosityLogs_IsChecked() const;
-
-	/** Toggles Verbosity "Warnings" true/false. */
-	void VerbosityWarnings_Execute();
+	ELogLevelFilter GetMessagesFilter() const;
+	void OnMessagesFilterChanged(ELogLevelFilter NewFilter);
 
-	/** Returns the state of Verbosity "Warnings". */
-	ECheckBoxState VerbosityWarnings_IsChecked() const;
-
-	/** Toggles Verbosity "Errors" true/false. */
-	void VerbosityErrors_Execute();
-
-	/** Returns the state of Verbosity "Errors". */
-	ECheckBoxState VerbosityErrors_IsChecked() const;
+	ELogLevelFilter GetWarningsFilter() const;
+	void OnWarningsFilterChanged(ELogLevelFilter NewFilter);
 
+	ELogLevelFilter GetErrorsFilter() const;
+	void OnErrorsFilterChanged(ELogLevelFilter NewFilter);
+	
 	/** Toggles All Categories true/false. */
 	void CategoriesShowAll_Execute();
 
 	/** Returns the state of "Show All" */
-	bool CategoriesShowAll_IsChecked() const;
+	ECheckBoxState CategoriesShowAll_IsChecked() const;
 
 	/** Toggles the given category true/false. */
 	void CategoriesSingle_Execute(FName InName);
diff --git a/Engine/Source/Developer/OutputLog/Public/OutputLogModule.h b/Engine/Source/Developer/OutputLog/Public/OutputLogModule.h
--- a/Engine/Source/Developer/OutputLog/Public/OutputLogModule.h
+++ b/Engine/Source/Developer/OutputLog/Public/OutputLogModule.h
@@ -106,7 +106,7 @@ public:
 
 private:
 	TSharedRef<SDockTab> SpawnOutputLogTab(const FSpawnTabArgs& Args);
-	void OnOutputLogTabClosed(TSharedRef<SDockTab> Tab);
+	void SaveDockedTabSettings();
 
 	TSharedRef<SDockTab> SpawnDeviceOutputLogTab(const FSpawnTabArgs& Args);
 
diff --git a/Engine/Source/Developer/OutputLog/Public/OutputLogSettings.h b/Engine/Source/Developer/OutputLog/Public/OutputLogSettings.h
--- a/Engine/Source/Developer/OutputLog/Public/OutputLogSettings.h
+++ b/Engine/Source/Developer/OutputLog/Public/OutputLogSettings.h
@@ -26,6 +26,58 @@ enum class ELogCategoryColorizationMode : uint8
 	ColorizeCategoryAsBadge
 };
 
+UENUM()
+enum class ELogLevelFilter : uint8
+{
+	/** Show none of the logs at this level. */
+	None,
+	/** Show only the enabled logs at this level. */
+	Enabled,
+	/** Show all logs at this level. */
+	All,
+};
+
+USTRUCT()
+struct FOutputLogCategorySettings
+{
+	GENERATED_BODY()
+	
+	UPROPERTY()
+	FName Name;
+
+	UPROPERTY()
+	bool bEnabled;
+};
+
+USTRUCT()
+struct FOutputLogFilterSettings
+{
+	GENERATED_BODY()
+
+	UPROPERTY()
+	ELogLevelFilter MessagesFilter = ELogLevelFilter::Enabled;
+
+	UPROPERTY()
+	ELogLevelFilter WarningsFilter = ELogLevelFilter::Enabled;
+
+	UPROPERTY()
+	ELogLevelFilter ErrorsFilter = ELogLevelFilter::Enabled;
+
+	UPROPERTY()
+	FText FilterText;
+
+	/**
+	 * Note that an empty list in settings implicitly means "all" categories.
+	 * This will invert state if a user explicitly disables all categories,
+	 * but that is acceptable. An empty log simply looks broken.
+	 */
+	UPROPERTY()
+	TArray<FOutputLogCategorySettings> Categories;
+
+	UPROPERTY()
+	bool bSelectNewCategories = true;
+};
+
 /**
  * Implements the Editor style settings.
  */
@@ -69,6 +121,10 @@ public:
 	bool bEnableOutputLogClearOnPIE;
 #endif
 
+	/** The most recently used filter settings. */
+	UPROPERTY(config)
+	FOutputLogFilterSettings OutputLogTabFilter;
+
 public:
 
 	/**
