Add a missing property initializer for FOutputLogCategorySettings::bEnabled.

diff --git a/Engine/Source/Developer/OutputLog/Public/OutputLogSettings.h b/Engine/Source/Developer/OutputLog/Public/OutputLogSettings.h
--- a/Engine/Source/Developer/OutputLog/Public/OutputLogSettings.h
+++ b/Engine/Source/Developer/OutputLog/Public/OutputLogSettings.h
@@ -46,7 +46,7 @@ struct FOutputLogCategorySettings
 	FName Name;
 
 	UPROPERTY()
-	bool bEnabled;
+	bool bEnabled = true;
 };
 
 USTRUCT()
