Updated SConsoleInputBox::GetCommandExecutorMenuContent to use UToolMenus

#rnx

diff --git a/Engine/Source/Developer/OutputLog/Private/OutputLogMenuContext.h b/Engine/Source/Developer/OutputLog/Private/OutputLogMenuContext.h
--- a/Engine/Source/Developer/OutputLog/Private/OutputLogMenuContext.h
+++ b/Engine/Source/Developer/OutputLog/Private/OutputLogMenuContext.h
@@ -8,9 +8,31 @@
 #include "OutputLogMenuContext.generated.h"
 
 
+class SConsoleInputBox;
 class SOutputLog;
 
 
+UCLASS()
+class UConsoleInputBoxMenuContext : public UObject
+{
+	GENERATED_BODY()
+
+public:
+	void Init(const TSharedRef<SConsoleInputBox>& InConsoleInputBox)
+	{
+		WeakConsoleInputBox = InConsoleInputBox;
+	}
+
+	TSharedPtr<SConsoleInputBox> GetConsoleInputBox() const
+	{
+		return WeakConsoleInputBox.Pin();
+	}
+
+private:
+	TWeakPtr<SConsoleInputBox> WeakConsoleInputBox;
+};
+
+
 UCLASS()
 class UOutputLogMenuContext : public UObject
 {
diff --git a/Engine/Source/Developer/OutputLog/Private/SOutputLog.cpp b/Engine/Source/Developer/OutputLog/Private/SOutputLog.cpp
--- a/Engine/Source/Developer/OutputLog/Private/SOutputLog.cpp
+++ b/Engine/Source/Developer/OutputLog/Private/SOutputLog.cpp
@@ -733,37 +733,52 @@ void SConsoleInputBox::MakeNextCommandExecutorActive()
 
 TSharedRef<SWidget> SConsoleInputBox::GetCommandExecutorMenuContent()
 {
-	FMenuBuilder MenuBuilder(/*bInShouldCloseWindowAfterMenuSelection=*/true, nullptr);
-
-	MenuBuilder.BeginSection("CmdExecEntries");
+	static const FName MenuName = "OutputLog.ConsoleInputBox.CmdExecMenu";
+	if (!UToolMenus::Get()->IsMenuRegistered(MenuName))
 	{
-		TArray<IConsoleCommandExecutor*> CommandExecutors = IModularFeatures::Get().GetModularFeatureImplementations<IConsoleCommandExecutor>(IConsoleCommandExecutor::ModularFeatureName());
-		CommandExecutors.Sort([](IConsoleCommandExecutor& LHS, IConsoleCommandExecutor& RHS)
-		{
-			return LHS.GetDisplayName().CompareTo(RHS.GetDisplayName()) < 0;
-		});
+		UToolMenu* Menu = UToolMenus::Get()->RegisterMenu(MenuName);
+		Menu->bShouldCloseWindowAfterMenuSelection = true;
 
-		for (const IConsoleCommandExecutor* CommandExecutor : CommandExecutors)
+		Menu->AddDynamicSection("DynamicCmdExecEntries", FNewToolMenuDelegate::CreateLambda([](UToolMenu* InMenu)
 		{
-			const bool bIsActiveCmdExec = ActiveCommandExecutor == CommandExecutor;
+			if (UConsoleInputBoxMenuContext* Context = InMenu->FindContext<UConsoleInputBoxMenuContext>())
+			{
+				if (TSharedPtr<SConsoleInputBox> This = Context->GetConsoleInputBox())
+				{
+					TArray<IConsoleCommandExecutor*> CommandExecutors = IModularFeatures::Get().GetModularFeatureImplementations<IConsoleCommandExecutor>(IConsoleCommandExecutor::ModularFeatureName());
+					CommandExecutors.Sort([](IConsoleCommandExecutor& LHS, IConsoleCommandExecutor& RHS)
+					{
+						return LHS.GetDisplayName().CompareTo(RHS.GetDisplayName()) < 0;
+					});
 
-			MenuBuilder.AddMenuEntry(
-				CommandExecutor->GetDisplayName(),
-				CommandExecutor->GetDescription(),
-				FSlateIcon(),
-				FUIAction(
-					FExecuteAction::CreateSP(this, &SConsoleInputBox::SetActiveCommandExecutor, CommandExecutor->GetName()),
-					FCanExecuteAction::CreateLambda([] { return true; }),
-					FIsActionChecked::CreateLambda([bIsActiveCmdExec] { return bIsActiveCmdExec; })
-					),
-				NAME_None,
-				EUserInterfaceActionType::Check
-			);
-		}
+					FToolMenuSection& Section = InMenu->AddSection("CmdExecEntries");
+					for (const IConsoleCommandExecutor* CommandExecutor : CommandExecutors)
+					{
+						const bool bIsActiveCmdExec = This->ActiveCommandExecutor == CommandExecutor;
+
+						Section.AddMenuEntry(
+							CommandExecutor->GetName(),
+							CommandExecutor->GetDisplayName(),
+							CommandExecutor->GetDescription(),
+							FSlateIcon(),
+							FUIAction(
+								FExecuteAction::CreateSP(This.Get(), &SConsoleInputBox::SetActiveCommandExecutor, CommandExecutor->GetName()),
+								FCanExecuteAction::CreateLambda([] { return true; }),
+								FIsActionChecked::CreateLambda([bIsActiveCmdExec] { return bIsActiveCmdExec; })
+								),
+							EUserInterfaceActionType::Check
+						);
+					}
+				}
+			}
+		}));
 	}
-	MenuBuilder.EndSection();
 
-	return MenuBuilder.MakeWidget();
+	UConsoleInputBoxMenuContext* MenuContext = NewObject<UConsoleInputBoxMenuContext>();
+	MenuContext->Init(SharedThis(this));
+
+	FToolMenuContext ToolMenuContext(MenuContext);
+	return UToolMenus::Get()->GenerateWidget(MenuName, ToolMenuContext);
 }
 
 FReply SConsoleInputBox::OnKeyDownHandler(const FGeometry& MyGeometry, const FKeyEvent& InKeyEvent)
