[OutputLog] SOutputLog: added AllowLogCategoryCallback to FOutputLogFilter so we can check when a new message with a new category is received if the category should be added  or not.

diff --git a/Engine/Source/Developer/OutputLog/Private/SOutputLog.cpp b/Engine/Source/Developer/OutputLog/Private/SOutputLog.cpp
--- a/Engine/Source/Developer/OutputLog/Private/SOutputLog.cpp
+++ b/Engine/Source/Developer/OutputLog/Private/SOutputLog.cpp
@@ -1779,6 +1779,8 @@ void SOutputLog::SetClearOnPIE(ECheckBoxState InValue)
 
 void SOutputLog::BuildInitialLogCategoryFilter(const FArguments& InArgs)
 {
+	Filter.AllowLogCategoryCallback = InArgs._AllowInitialLogCategory;
+
 	for (const auto& Message : InArgs._Messages)
 	{
 		const bool bIsDeselectedByDefault = InArgs._AllowInitialLogCategory.IsBound() && !InArgs._AllowInitialLogCategory.Execute(Message->Category);
@@ -2515,7 +2517,12 @@ void FOutputLogFilter::AddAvailableLogCategory(const FName& LogCategory, TOption
 			break;
 		}
 	}
-	Categories.Insert({ LogCategory, InitiallySelected.Get(bSelectNewCategories) }, InsertIndex+1);
+
+	const bool bEnabled = InitiallySelected.IsSet()
+		? InitiallySelected.GetValue()
+		: (bSelectNewCategories && (!AllowLogCategoryCallback.IsBound() || AllowLogCategoryCallback.Execute(LogCategory)));
+
+	Categories.Insert({ LogCategory, bEnabled }, InsertIndex + 1);
 }
 
 bool FOutputLogFilter::IsLogCategoryAvailable(const FName& LogCategory) const
diff --git a/Engine/Source/Developer/OutputLog/Private/SOutputLog.h b/Engine/Source/Developer/OutputLog/Private/SOutputLog.h
--- a/Engine/Source/Developer/OutputLog/Private/SOutputLog.h
+++ b/Engine/Source/Developer/OutputLog/Private/SOutputLog.h
@@ -253,6 +253,9 @@ struct FOutputLogFilter
 	/** Whether newly added categories should begin as selected. */
 	bool bSelectNewCategories = true;
 
+	/** If set, checks if a category is allowed to log or not */
+	FAllowLogCategoryCallback AllowLogCategoryCallback;
+
 	/** Enable all filters by default */
 	FOutputLogFilter() : TextFilterExpressionEvaluator(ETextFilterExpressionEvaluatorMode::BasicString)
 	{
