[OutputLog] Have the Dock In Layout button still show when the output log is docked, and close the drawer tab / focus on the docked tab if it already exists when the button is clicked

diff --git a/Engine/Source/Developer/OutputLog/Private/SOutputLog.cpp b/Engine/Source/Developer/OutputLog/Private/SOutputLog.cpp
--- a/Engine/Source/Developer/OutputLog/Private/SOutputLog.cpp
+++ b/Engine/Source/Developer/OutputLog/Private/SOutputLog.cpp
@@ -1216,6 +1216,8 @@ void SOutputLog::Construct( const FArguments& InArgs, bool bCreateDrawerDockButt
 	bEnableLoggingLimit = InArgs._LoggingLineLimit.IsSet();
 	LoggingLineLimit = InArgs._LoggingLineLimit.Get(10000);
 
+	OnCloseConsole = InArgs._OnCloseConsole;
+
 	MessagesTextMarshaller = FOutputLogTextLayoutMarshaller::Create(InArgs._Messages, &Filter);
 
 	MessagesTextBox = SNew(SMultiLineEditableTextBox)
@@ -1349,7 +1351,7 @@ void SOutputLog::Construct( const FArguments& InArgs, bool bCreateDrawerDockButt
 				return FOutputLogModule::Get().ShouldHideConsole() ? EVisibility::Collapsed : EVisibility::Visible;
 			}))
 			.OnConsoleCommandExecuted(this, &SOutputLog::OnConsoleCommandExecuted)
-			.OnCloseConsole(InArgs._OnCloseConsole)
+			.OnCloseConsole(OnCloseConsole)
 			// Always place suggestions above the input line for the output log widget
 			.SuggestionListPlacement(MenuPlacement_AboveAnchor) 
 		]
@@ -2315,13 +2317,14 @@ TSharedRef<SWidget> SOutputLog::CreateDrawerDockButton()
 		return
 			SNew(SButton)
 			.ButtonStyle(FOutputLogStyle::Get(), "SimpleButton")
-			.ToolTipText(LOCTEXT("DockInLayout_Tooltip", "Docks this output log in the current layout.\nThe drawer will still be usable as a temporary log."))
-			.ContentPadding(FMargin(1, 0))
-			.Visibility_Lambda(
+			.ToolTipText_Lambda(
 				[]()
 				{
-					return FOutputLogModule::Get().GetOutputLogTab() == nullptr ? EVisibility::Visible : EVisibility::Hidden;
+					return !FOutputLogModule::Get().GetOutputLogTab().IsValid() ?
+						LOCTEXT("DockInLayout_Tooltip", "Docks this output log in the current layout.\nThe drawer will still be usable as a temporary log.") :
+						LOCTEXT("FocusOnDocked_Tooltip", "Close the drawer and focus on the currently docked output log.");
 				})
+			.ContentPadding(FMargin(1, 0))
 			.OnClicked(this, &SOutputLog::OnDockInLayoutClicked)
 			[
 				SNew(SHorizontalBox)
@@ -2339,7 +2342,13 @@ TSharedRef<SWidget> SOutputLog::CreateDrawerDockButton()
 				.Padding(4.0, 0.0f)
 				[
 					SNew(STextBlock)
-					.Text(LOCTEXT("DockInLayout", "Dock in Layout"))
+					.Text_Lambda(
+					[]()
+					{
+						return !FOutputLogModule::Get().GetOutputLogTab().IsValid() ?
+							LOCTEXT("DockInLayout", "Dock in Layout") :
+							LOCTEXT("FocusOnDocked", "Focus On Docked");
+					})
 					.ColorAndOpacity(FSlateColor::UseForeground())
 				]
 			];
@@ -2393,6 +2402,8 @@ FReply SOutputLog::OnDockInLayoutClicked()
 		FGlobalTabmanager::Get()->TryInvokeTab(OutputLogTabName);
 	}
 
+	OnCloseConsole.ExecuteIfBound();
+
 	return FReply::Handled();
 }
 
diff --git a/Engine/Source/Developer/OutputLog/Private/SOutputLog.h b/Engine/Source/Developer/OutputLog/Private/SOutputLog.h
--- a/Engine/Source/Developer/OutputLog/Private/SOutputLog.h
+++ b/Engine/Source/Developer/OutputLog/Private/SOutputLog.h
@@ -438,6 +438,7 @@ protected:
 	bool bIsUserScrolled;
 
 	FSimpleDelegate OnClearLogDelegate; 
+    FSimpleDelegate OnCloseConsole;
 
 private:
 
