Removed left over pragma optimize from SOutputLog.cpp

diff --git a/Engine/Source/Developer/OutputLog/Private/SOutputLog.cpp b/Engine/Source/Developer/OutputLog/Private/SOutputLog.cpp
--- a/Engine/Source/Developer/OutputLog/Private/SOutputLog.cpp
+++ b/Engine/Source/Developer/OutputLog/Private/SOutputLog.cpp
@@ -2614,4 +2614,3 @@ void FOutputLogFilter::ClearSelectedLogCategories()
 }
 #undef LOCTEXT_NAMESPACE
 
-#pragma optimize("", on)
