[OutputLog] SOutputLog: use binary search in FOutputLogFilter's AddAvailableLogCategory and FindCategoryFilter methods

diff --git a/Engine/Source/Developer/OutputLog/Private/SOutputLog.cpp b/Engine/Source/Developer/OutputLog/Private/SOutputLog.cpp
--- a/Engine/Source/Developer/OutputLog/Private/SOutputLog.cpp
+++ b/Engine/Source/Developer/OutputLog/Private/SOutputLog.cpp
@@ -1,6 +1,8 @@
 // Copyright Epic Games, Inc. All Rights Reserved.
 
 #include "SOutputLog.h"
+
+#include "Algo/BinarySearch.h"
 #include "ConsoleSettings.h"
 #include "Framework/Commands/UICommandInfo.h"
 #include "Framework/Text/IRun.h"
@@ -2513,27 +2515,23 @@ ELogLevelFilter FOutputLogFilter::GetMessageLevelFilter(const TSharedPtr<FOutput
 
 void FOutputLogFilter::AddAvailableLogCategory(const FName& LogCategory, TOptional<bool> InitiallySelected)
 {
-	// Use an insert-sort to keep AvailableLogCategories alphabetically sorted
-	int32 InsertIndex = 0;
-	for (InsertIndex = Categories.Num()-1; InsertIndex >= 0; --InsertIndex)
-	{
-		FOutputLogCategorySettings& CheckCategory = Categories[InsertIndex];
-		// No duplicates
-		if (CheckCategory.Name == LogCategory)
-		{
-			return;
-		}
-		else if (CheckCategory.Name.Compare(LogCategory) < 0)
+	const int32 InsertIndex = Algo::LowerBoundBy(Categories, LogCategory, &FOutputLogCategorySettings::Name,
+		[](const FName& A, const FName& B)
 		{
-			break;
-		}
+			return A.Compare(B) < 0;
+		});
+
+	const bool bCategoryExists = InsertIndex < Categories.Num() && Categories[InsertIndex].Name == LogCategory;
+	if (bCategoryExists)
+	{
+		return;
 	}
 
 	const bool bEnabled = InitiallySelected.IsSet()
 		? InitiallySelected.GetValue()
 		: (bSelectNewCategories && (!AllowLogCategoryCallback.IsBound() || AllowLogCategoryCallback.Execute(LogCategory)));
 
-	Categories.Insert({ LogCategory, bEnabled }, InsertIndex + 1);
+	Categories.EmplaceAt(InsertIndex, LogCategory, bEnabled);
 }
 
 bool FOutputLogFilter::IsLogCategoryAvailable(const FName& LogCategory) const
@@ -2543,28 +2541,24 @@ bool FOutputLogFilter::IsLogCategoryAvailable(const FName& LogCategory) const
 
 FOutputLogCategorySettings* FOutputLogFilter::FindCategoryFilter(const FName& LogCategory)
 {
-	// TODO: Could in theory use a binary search as AddAvailableLogCategory() ensures `Categories` will be sorted.
-	for (FOutputLogCategorySettings& Item : Categories)
-	{
-		if (Item.Name == LogCategory)
+	const int32 FoundIndex = Algo::BinarySearchBy(Categories, LogCategory, &FOutputLogCategorySettings::Name,
+		[](const FName& A, const FName& B)
 		{
-			return &Item;
-		}
-	}
-	return nullptr;
+			return A.Compare(B) < 0;
+		});
+	
+	return FoundIndex != INDEX_NONE ? &Categories[FoundIndex] : nullptr;
 }
 
 const FOutputLogCategorySettings* FOutputLogFilter::FindCategoryFilter(const FName& LogCategory) const
 {
-	// TODO: Could in theory use a binary search as AddAvailableLogCategory() ensures `Categories` will be sorted.
-	for (const FOutputLogCategorySettings& Item : Categories)
-	{
-		if (Item.Name == LogCategory)
+	const int32 FoundIndex = Algo::BinarySearchBy(Categories, LogCategory, &FOutputLogCategorySettings::Name,
+		[](const FName& A, const FName& B)
 		{
-			return &Item;
-		}
-	}
-	return nullptr;
+			return A.Compare(B) < 0;
+		});
+
+	return FoundIndex != INDEX_NONE ? &Categories[FoundIndex] : nullptr;
 }
 
 void FOutputLogFilter::SetLogCategoryEnabled(const FName& LogCategory, bool bEnabled)
