OutputLog - Adding profiling events + minor adjustments

#Synced-CL 50897806

diff --git a/Engine/Source/Developer/OutputLog/Private/SOutputLog.cpp b/Engine/Source/Developer/OutputLog/Private/SOutputLog.cpp
--- a/Engine/Source/Developer/OutputLog/Private/SOutputLog.cpp
+++ b/Engine/Source/Developer/OutputLog/Private/SOutputLog.cpp
@@ -43,6 +43,7 @@
 #include "OutputLogMenuContext.h"
 #include "ToolMenus.h"
 #include "Widgets/Input/SSegmentedControl.h"
+#include "Trace/SlateTrace.h"
 
 
 #define LOCTEXT_NAMESPACE "SOutputLog"
@@ -932,6 +933,8 @@ void FOutputLogTextLayoutMarshaller::AppendPendingMessagesToTextLayout()
 		return;
 	}
 
+	TRACE_CPUPROFILER_EVENT_SCOPE_ON_CHANNEL(FOutputLogTextLayoutMarshaller::AppendPendingMessagesToTextLayout, SlateChannel);
+
 	if (TextLayout)
 	{
 		// If we were previously empty, then we'd have inserted a dummy empty line into the document
@@ -1188,6 +1191,7 @@ FOutputLogTextLayoutMarshaller::FOutputLogTextLayoutMarshaller(TArray< TSharedPt
 	: Messages(MoveTemp(InMessages))
 	, NextPendingMessageIndex(0)
 	, CachedNumMessages(0)
+	, bNumMessagesCacheDirty(false)
 	, Filter(InFilter)
 	, TextLayout(nullptr)
 {
@@ -1212,6 +1216,8 @@ namespace
 BEGIN_SLATE_FUNCTION_BUILD_OPTIMIZATION
 void SOutputLog::Construct( const FArguments& InArgs, bool bCreateDrawerDockButton)
 {
+	TRACE_CPUPROFILER_EVENT_SCOPE_ON_CHANNEL(SOutputLog::Construct, SlateChannel);
+
 	bShouldCreateDrawerDockButton = bCreateDrawerDockButton;
 	BuildInitialLogCategoryFilter(InArgs);
 
@@ -1678,6 +1684,8 @@ void SOutputLog::RequestForceScroll(bool bIfUserHasNotScrolledUp)
 
 void SOutputLog::Refresh()
 {
+	TRACE_CPUPROFILER_EVENT_SCOPE_ON_CHANNEL(SOutputLog::Refresh, SlateChannel);
+
 	// Re-count messages if filter changed before we refresh
 	MessagesTextMarshaller->CountMessages();
 
@@ -2516,6 +2524,8 @@ ELogLevelFilter FOutputLogFilter::GetMessageLevelFilter(const TSharedPtr<FOutput
 
 void FOutputLogFilter::AddAvailableLogCategory(const FName& LogCategory, TOptional<bool> InitiallySelected)
 {
+	TRACE_CPUPROFILER_EVENT_SCOPE_ON_CHANNEL(FOutputLogFilter::AddAvailableLogCategory, SlateChannel);
+
 	const int32 InsertIndex = Algo::LowerBoundBy(Categories, LogCategory, &FOutputLogCategorySettings::Name,
 		[](const FName& A, const FName& B)
 		{
diff --git a/Engine/Source/Developer/OutputLog/Private/SOutputLog.h b/Engine/Source/Developer/OutputLog/Private/SOutputLog.h
--- a/Engine/Source/Developer/OutputLog/Private/SOutputLog.h
+++ b/Engine/Source/Developer/OutputLog/Private/SOutputLog.h
@@ -273,13 +273,13 @@ struct FOutputLogFilter
 	void SetFilterText(const FText& InFilterText) { TextFilterExpressionEvaluator.SetFilterText(InFilterText); }
 
 	/** Get the Text currently being used as the Filter's restrictions */
-	const FText GetFilterText() const { return TextFilterExpressionEvaluator.GetFilterText(); }
+	FText GetFilterText() const { return TextFilterExpressionEvaluator.GetFilterText(); }
 
 	/** Returns Evaluator syntax errors (if any) */
-	FText GetSyntaxErrors() { return TextFilterExpressionEvaluator.GetFilterErrorText(); }
+	FText GetSyntaxErrors() const { return TextFilterExpressionEvaluator.GetFilterErrorText(); }
 
 	const TArray<FOutputLogCategorySettings>& GetCategoryFilters() const { return Categories; }
-	
+
 	/** Adds a Log Category to the list of available categories, if it isn't already present */
 	void AddAvailableLogCategory(const FName& LogCategory, TOptional<bool> InitiallySelected = TOptional<bool>());
 
@@ -588,7 +588,7 @@ protected:
 
 	/** Holds cached numbers of messages to avoid unnecessary re-filtering */
 	int32 CachedNumMessages;
-	
+
 	/** Flag indicating the messages count cache needs rebuilding */
 	bool bNumMessagesCacheDirty;
 
