Fix compilation issue due to missing #if

#tests Compiled

#Synced-CL 50880041

diff --git a/Engine/Source/Developer/OutputLog/Private/SOutputLog.cpp b/Engine/Source/Developer/OutputLog/Private/SOutputLog.cpp
--- a/Engine/Source/Developer/OutputLog/Private/SOutputLog.cpp
+++ b/Engine/Source/Developer/OutputLog/Private/SOutputLog.cpp
@@ -933,7 +933,9 @@ void FOutputLogTextLayoutMarshaller::AppendPendingMessagesToTextLayout()
 		return;
 	}
 
+#if CPUPROFILERTRACE_ENABLED
 	TRACE_CPUPROFILER_EVENT_SCOPE_ON_CHANNEL(FOutputLogTextLayoutMarshaller::AppendPendingMessagesToTextLayout, SlateChannel);
+#endif
 
 	if (TextLayout)
 	{
@@ -1216,7 +1218,9 @@ namespace
 BEGIN_SLATE_FUNCTION_BUILD_OPTIMIZATION
 void SOutputLog::Construct( const FArguments& InArgs, bool bCreateDrawerDockButton)
 {
+#if CPUPROFILERTRACE_ENABLED
 	TRACE_CPUPROFILER_EVENT_SCOPE_ON_CHANNEL(SOutputLog::Construct, SlateChannel);
+#endif
 
 	bShouldCreateDrawerDockButton = bCreateDrawerDockButton;
 	BuildInitialLogCategoryFilter(InArgs);
@@ -1684,7 +1688,9 @@ void SOutputLog::RequestForceScroll(bool bIfUserHasNotScrolledUp)
 
 void SOutputLog::Refresh()
 {
+#if CPUPROFILERTRACE_ENABLED
 	TRACE_CPUPROFILER_EVENT_SCOPE_ON_CHANNEL(SOutputLog::Refresh, SlateChannel);
+#endif
 
 	// Re-count messages if filter changed before we refresh
 	MessagesTextMarshaller->CountMessages();
@@ -2524,7 +2530,9 @@ ELogLevelFilter FOutputLogFilter::GetMessageLevelFilter(const TSharedPtr<FOutput
 
 void FOutputLogFilter::AddAvailableLogCategory(const FName& LogCategory, TOptional<bool> InitiallySelected)
 {
+#if CPUPROFILERTRACE_ENABLED
 	TRACE_CPUPROFILER_EVENT_SCOPE_ON_CHANNEL(FOutputLogFilter::AddAvailableLogCategory, SlateChannel);
+#endif
 
 	const int32 InsertIndex = Algo::LowerBoundBy(Categories, LogCategory, &FOutputLogCategorySettings::Name,
 		[](const FName& A, const FName& B)
