[Fix] Fixing compilation issues for program targets due to missing

#Synced-CL 50897806

diff --git a/Engine/Source/Developer/OutputLog/Private/SOutputLog.cpp b/Engine/Source/Developer/OutputLog/Private/SOutputLog.cpp
--- a/Engine/Source/Developer/OutputLog/Private/SOutputLog.cpp
+++ b/Engine/Source/Developer/OutputLog/Private/SOutputLog.cpp
@@ -933,7 +933,7 @@ void FOutputLogTextLayoutMarshaller::AppendPendingMessagesToTextLayout()
 		return;
 	}
 
-#if CPUPROFILERTRACE_ENABLED
+#if UE_SLATE_TRACE_ENABLED
 	TRACE_CPUPROFILER_EVENT_SCOPE_ON_CHANNEL(FOutputLogTextLayoutMarshaller::AppendPendingMessagesToTextLayout, SlateChannel);
 #endif
 
@@ -1218,7 +1218,7 @@ namespace
 BEGIN_SLATE_FUNCTION_BUILD_OPTIMIZATION
 void SOutputLog::Construct( const FArguments& InArgs, bool bCreateDrawerDockButton)
 {
-#if CPUPROFILERTRACE_ENABLED
+#if UE_SLATE_TRACE_ENABLED
 	TRACE_CPUPROFILER_EVENT_SCOPE_ON_CHANNEL(SOutputLog::Construct, SlateChannel);
 #endif
 
@@ -1688,7 +1688,7 @@ void SOutputLog::RequestForceScroll(bool bIfUserHasNotScrolledUp)
 
 void SOutputLog::Refresh()
 {
-#if CPUPROFILERTRACE_ENABLED
+#if UE_SLATE_TRACE_ENABLED
 	TRACE_CPUPROFILER_EVENT_SCOPE_ON_CHANNEL(SOutputLog::Refresh, SlateChannel);
 #endif
 
@@ -2530,7 +2530,7 @@ ELogLevelFilter FOutputLogFilter::GetMessageLevelFilter(const TSharedPtr<FOutput
 
 void FOutputLogFilter::AddAvailableLogCategory(const FName& LogCategory, TOptional<bool> InitiallySelected)
 {
-#if CPUPROFILERTRACE_ENABLED
+#if UE_SLATE_TRACE_ENABLED
 	TRACE_CPUPROFILER_EVENT_SCOPE_ON_CHANNEL(FOutputLogFilter::AddAvailableLogCategory, SlateChannel);
 #endif
 
