OutputLog - Improving Log Categories performance

Currently when opening the Output Log, all the messages are processed 3 times just to create the Log Categories.
A first time in `BuildInitialLogCategoryFilter`, then when the `SMultilineEditableText` is created which ends up calling `FSlateEditableTextLayout::SetEditableText` > `FOutputLogTextLayoutMarshaller::SubmitPendingMessages`, and finally in `SOutputLog::Refresh` which ends up calling again `FSlateEditableTextLayout::SetEditableText` > `FOutputLogTextLayoutMarshaller::SubmitPendingMessages`.
The constant sorting is also quite slow.

To allieviate the issue:
- `BuildInitialLogCategoryFilter` will now only process Categories with display overrides
- The sorting on insertion has been replaced with a sorting when retrieving the list of categories (which does not happen often).
- The `TArray` was replaced with a `TMap` for faster retrieval
- Some additional arguments were added to `SOutputLog::FArguments` to avoid the call to `SOutputLog::Refresh`.

On an Output Log with ~20000 messages, time to append the categories when opening the Output Log went down from ~650ms to ~100ms

#Synced-CL 50897806

diff --git a/Engine/Source/Developer/OutputLog/Private/OutputLogModule.cpp b/Engine/Source/Developer/OutputLog/Private/OutputLogModule.cpp
--- a/Engine/Source/Developer/OutputLog/Private/OutputLogModule.cpp
+++ b/Engine/Source/Developer/OutputLog/Private/OutputLogModule.cpp
@@ -110,8 +110,9 @@ private:
 
 TSharedRef<SDockTab> FOutputLogModule::SpawnOutputLogTab(const FSpawnTabArgs& Args)
 {
-	TSharedRef<SOutputLog> NewLog = SNew(SOutputLog, false).Messages(OutputLogHistory->GetMessages());
-	NewLog->UpdateOutputLogFilter(GetDefault<UOutputLogSettings>()->OutputLogTabFilter);
+	TSharedRef<SOutputLog> NewLog = SNew(SOutputLog, false)
+			.Messages(&OutputLogHistory->GetMessages())
+			.InitialLogFilterSettings(&GetDefault<UOutputLogSettings>()->OutputLogTabFilter);
 
 	OutputLog = NewLog;
 
@@ -120,7 +121,7 @@ TSharedRef<SDockTab> FOutputLogModule::SpawnOutputLogTab(const FSpawnTabArgs& Ar
 		.Label(NSLOCTEXT("OutputLog", "TabTitle", "Output Log"))
 		.OnPersistVisualState_Raw(this, &FOutputLogModule::SaveDockedTabSettings)
 		[
-			NewLog
+			MoveTemp(NewLog)
 		];
 
 	OutputLogTab = NewTab;
@@ -249,9 +250,8 @@ TSharedRef<SWidget> FOutputLogModule::MakeOutputLogDrawerWidget(const FSimpleDel
 					OutputLogModule->OnOutputLogDrawerCleared.Broadcast();
 				}
 			})
-			.Messages(OutputLogHistory->GetMessages());
-
-		OutputLogDrawerPinned->UpdateOutputLogFilter(*OutputLogFilterCache);
+			.Messages(&OutputLogHistory->GetMessages())
+			.InitialLogFilter(OutputLogFilterCache.Get());
 
 		OutputLogDrawer = OutputLogDrawerPinned;
 	}
@@ -263,7 +263,7 @@ TSharedRef<SWidget> FOutputLogModule::MakeOutputLogWidget(const FOutputLogCreati
 {
 	return SNew(SOutputLog, Params.bCreateDockInLayoutButton)
 			.OnCloseConsole(Params.OnCloseConsole)
-			.Messages(OutputLogHistory->GetMessages())
+			.Messages(&OutputLogHistory->GetMessages())
 			.SettingsMenuFlags(Params.SettingsMenuCreationFlags)
 			.DefaultCategorySelection(Params.DefaultCategorySelection)
 			.AllowInitialLogCategory(Params.AllowAsInitialLogCategory)
diff --git a/Engine/Source/Developer/OutputLog/Private/SOutputLog.cpp b/Engine/Source/Developer/OutputLog/Private/SOutputLog.cpp
--- a/Engine/Source/Developer/OutputLog/Private/SOutputLog.cpp
+++ b/Engine/Source/Developer/OutputLog/Private/SOutputLog.cpp
@@ -974,7 +974,7 @@ void FOutputLogTextLayoutMarshaller::AppendPendingMessagesToTextLayout()
 
 	for (int32 MessageIndex = NextPendingMessageIndex; MessageIndex < CurrentMessagesCount; ++MessageIndex)
 	{
-		const TSharedPtr<FOutputLogMessage> Message = Messages[MessageIndex];
+		const TSharedPtr<FOutputLogMessage>& Message = Messages[MessageIndex];
 		const int32 LineIndex = TextLayout->GetLineModels().Num() + NumAddedMessages;
 
 		if (!Message)
@@ -1223,6 +1223,15 @@ void SOutputLog::Construct( const FArguments& InArgs, bool bCreateDrawerDockButt
 #endif
 
 	bShouldCreateDrawerDockButton = bCreateDrawerDockButton;
+
+	if (InArgs._InitialLogFilter)
+	{
+		Filter = *InArgs._InitialLogFilter;
+	}
+	if (InArgs._InitialLogFilterSettings)
+	{
+		Filter.ApplySettings(*InArgs._InitialLogFilterSettings);
+	}
 	BuildInitialLogCategoryFilter(InArgs);
 
 	bShouldShowLoggingLimitMenu = InArgs._EnableLoggingLimitMenu;
@@ -1231,7 +1240,7 @@ void SOutputLog::Construct( const FArguments& InArgs, bool bCreateDrawerDockButt
 
 	OnCloseConsole = InArgs._OnCloseConsole;
 
-	MessagesTextMarshaller = FOutputLogTextLayoutMarshaller::Create(InArgs._Messages, &Filter);
+	MessagesTextMarshaller = FOutputLogTextLayoutMarshaller::Create(InArgs._Messages ? *InArgs._Messages : TArray< TSharedPtr<FOutputLogMessage> >{}, &Filter);
 
 	MessagesTextBox = SNew(SMultiLineEditableTextBox)
 		.Style(FOutputLogStyle::Get(), "Log.TextBox")
@@ -1800,10 +1809,19 @@ void SOutputLog::BuildInitialLogCategoryFilter(const FArguments& InArgs)
 {
 	Filter.AllowLogCategoryCallback = InArgs._AllowInitialLogCategory;
 
-	for (const auto& Message : InArgs._Messages)
+	if (InArgs._Messages && InArgs._AllowInitialLogCategory.IsBound()) // There is no need to process the log categories now if this is not bound, as adding the messages will do it for us
 	{
-		const bool bIsDeselectedByDefault = InArgs._AllowInitialLogCategory.IsBound() && !InArgs._AllowInitialLogCategory.Execute(Message->Category);
-		Filter.AddAvailableLogCategory(Message->Category, bIsDeselectedByDefault ? false : TOptional<bool>());
+		bool bCategoryAlreadyInSet;
+		TSet<FName> CategoryNames;
+		for (const TSharedPtr<FOutputLogMessage>& Message : *InArgs._Messages)
+		{
+			CategoryNames.Add(Message->Category, &bCategoryAlreadyInSet);
+			if (!bCategoryAlreadyInSet)
+			{
+				const bool bIsDeselectedByDefault = !InArgs._AllowInitialLogCategory.Execute(Message->Category);
+				Filter.AddAvailableLogCategory(Message->Category, bIsDeselectedByDefault ? false : TOptional<bool>());
+			}
+		}
 	}
 
 	for (auto DefaultCategorySelectionIt = InArgs._DefaultCategorySelection.CreateConstIterator(); DefaultCategorySelectionIt; ++DefaultCategorySelectionIt)
@@ -2426,20 +2444,25 @@ FReply SOutputLog::OnDockInLayoutClicked()
 
 ECheckBoxState FOutputLogFilter::AreAllCategoriesSelected() const
 {
-	int32 Count = 0;
-	for (const FOutputLogCategorySettings& Item : Categories)
+	int32 NumCategoriesEnabled = 0;
+	for (const TPair<FName, FOutputLogCategorySettings>& Pair : Categories)
 	{
-		if (Item.bEnabled)
+		const FOutputLogCategorySettings& Category = Pair.Value;
+		if (Category.bEnabled)
 		{
-			Count++; 
+			++NumCategoriesEnabled;
+		}
+		else if (NumCategoriesEnabled > 0)
+		{
+			return ECheckBoxState::Undetermined; // We have both enabled and disabled ones, no need to continue
 		}
 	}
 	
-	if (Count == Categories.Num())
+	if (NumCategoriesEnabled == Categories.Num())
 	{
 		return ECheckBoxState::Checked;
 	}
-	if (Count == 0)
+	if (NumCategoriesEnabled == 0)
 	{
 		return ECheckBoxState::Unchecked;
 	}
@@ -2487,7 +2510,7 @@ void FOutputLogFilter::ExportSettings(FOutputLogFilterSettings& OutSettings) con
 	}
 	else
 	{
-		OutSettings.Categories = Categories;
+		OutSettings.Categories = GetCategoryFilters();
 	}
 }
 
@@ -2528,61 +2551,43 @@ ELogLevelFilter FOutputLogFilter::GetMessageLevelFilter(const TSharedPtr<FOutput
 	}
 }
 
-void FOutputLogFilter::AddAvailableLogCategory(const FName& LogCategory, TOptional<bool> InitiallySelected)
+TArray<FOutputLogCategorySettings> FOutputLogFilter::GetCategoryFilters() const
 {
 #if UE_SLATE_TRACE_ENABLED
-	TRACE_CPUPROFILER_EVENT_SCOPE_ON_CHANNEL(FOutputLogFilter::AddAvailableLogCategory, SlateChannel);
+	TRACE_CPUPROFILER_EVENT_SCOPE_ON_CHANNEL(FOutputLogFilter::GetCategoryFilters, SlateChannel);
 #endif
 
-	const int32 InsertIndex = Algo::LowerBoundBy(Categories, LogCategory, &FOutputLogCategorySettings::Name,
-		[](const FName& A, const FName& B)
-		{
-			return A.Compare(B) < 0;
-		});
-
-	const bool bCategoryExists = InsertIndex < Categories.Num() && Categories[InsertIndex].Name == LogCategory;
-	if (bCategoryExists)
+	if (bCategoriesNeedSorting)
 	{
-		return;
+		const_cast<FOutputLogFilter*>(this)->Categories.KeyStableSort(FNameLexicalLess());
+		bCategoriesNeedSorting = false;
 	}
-
-	const bool bEnabled = InitiallySelected.IsSet()
-		? InitiallySelected.GetValue()
-		: (bSelectNewCategories && (!AllowLogCategoryCallback.IsBound() || AllowLogCategoryCallback.Execute(LogCategory)));
-
-	Categories.EmplaceAt(InsertIndex, LogCategory, bEnabled);
+	TArray<FOutputLogCategorySettings> CategoryFilters;
+	Categories.GenerateValueArray(CategoryFilters);
+	return CategoryFilters;
 }
 
-bool FOutputLogFilter::IsLogCategoryAvailable(const FName& LogCategory) const
+void FOutputLogFilter::AddAvailableLogCategory(const FName& LogCategory, const TOptional<bool>& InitiallySelected)
 {
-	return FindCategoryFilter(LogCategory) != nullptr;
-}
+#if UE_SLATE_TRACE_ENABLED
+	TRACE_CPUPROFILER_EVENT_SCOPE_ON_CHANNEL(FOutputLogFilter::AddAvailableLogCategory, SlateChannel);
+#endif
 
-FOutputLogCategorySettings* FOutputLogFilter::FindCategoryFilter(const FName& LogCategory)
-{
-	const int32 FoundIndex = Algo::BinarySearchBy(Categories, LogCategory, &FOutputLogCategorySettings::Name,
-		[](const FName& A, const FName& B)
-		{
-			return A.Compare(B) < 0;
-		});
-	
-	return FoundIndex != INDEX_NONE ? &Categories[FoundIndex] : nullptr;
+	if (!Categories.Contains(LogCategory))
+	{
+		Categories.Add(LogCategory, { LogCategory, InitiallySelected.Get(bSelectNewCategories) });
+		bCategoriesNeedSorting = true;
+	}
 }
 
-const FOutputLogCategorySettings* FOutputLogFilter::FindCategoryFilter(const FName& LogCategory) const
+bool FOutputLogFilter::IsLogCategoryAvailable(const FName& LogCategory) const
 {
-	const int32 FoundIndex = Algo::BinarySearchBy(Categories, LogCategory, &FOutputLogCategorySettings::Name,
-		[](const FName& A, const FName& B)
-		{
-			return A.Compare(B) < 0;
-		});
-
-	return FoundIndex != INDEX_NONE ? &Categories[FoundIndex] : nullptr;
+	return Categories.Find(LogCategory) != nullptr;
 }
 
 void FOutputLogFilter::SetLogCategoryEnabled(const FName& LogCategory, bool bEnabled)
 {
-	if (FOutputLogCategorySettings* CategoryFilter = FindCategoryFilter(LogCategory))
+	if (FOutputLogCategorySettings* CategoryFilter = Categories.Find(LogCategory))
 	{
 		CategoryFilter->bEnabled = bEnabled;
 	}
@@ -2594,15 +2599,16 @@ void FOutputLogFilter::SetLogCategoryEnabled(const FName& LogCategory, bool bEna
 
 void FOutputLogFilter::SetAllCategoriesEnabled(bool bEnabled)
 {
-	for (FOutputLogCategorySettings& Item : Categories)
+	for (TPair<FName, FOutputLogCategorySettings>& Pair : Categories)
 	{
-		Item.bEnabled = bEnabled;
+		FOutputLogCategorySettings& Category = Pair.Value;
+		Category.bEnabled = bEnabled;
 	}
 }
 
 void FOutputLogFilter::ToggleLogCategory(const FName& LogCategory)
 {
-	if (FOutputLogCategorySettings* CategoryFilter = FindCategoryFilter(LogCategory))
+	if (FOutputLogCategorySettings* CategoryFilter = Categories.Find(LogCategory))
 	{
 		CategoryFilter->bEnabled = !CategoryFilter->bEnabled;
 	}
@@ -2614,7 +2620,7 @@ void FOutputLogFilter::ToggleLogCategory(const FName& LogCategory)
 
 bool FOutputLogFilter::IsLogCategoryEnabled(const FName& LogCategory) const
 {
-	if (const FOutputLogCategorySettings* CategoryFilter = FindCategoryFilter(LogCategory))
+	if (const FOutputLogCategorySettings* CategoryFilter = Categories.Find(LogCategory))
 	{
 		return CategoryFilter->bEnabled;
 	}
diff --git a/Engine/Source/Developer/OutputLog/Private/SOutputLog.h b/Engine/Source/Developer/OutputLog/Private/SOutputLog.h
--- a/Engine/Source/Developer/OutputLog/Private/SOutputLog.h
+++ b/Engine/Source/Developer/OutputLog/Private/SOutputLog.h
@@ -243,19 +243,19 @@ private:
 */
 struct FOutputLogFilter
 {
+	/** Set of Verbosity levels that will show all regardless of category filter */
+	TSet<ELogVerbosity::Type> IgnoreFilterVerbosities;
+
+	/** If set, checks if a category is allowed to log or not */
+	FAllowLogCategoryCallback AllowLogCategoryCallback;
+
 	ELogLevelFilter MessagesFilter = ELogLevelFilter::Enabled;
 	ELogLevelFilter WarningsFilter = ELogLevelFilter::Enabled;
 	ELogLevelFilter ErrorsFilter = ELogLevelFilter::Enabled;
 
-	/** Set of Verbosity levels that will show all regardless of category filter */
-	TSet<ELogVerbosity::Type> IgnoreFilterVerbosities;
-
 	/** Whether newly added categories should begin as selected. */
 	bool bSelectNewCategories = true;
 
-	/** If set, checks if a category is allowed to log or not */
-	FAllowLogCategoryCallback AllowLogCategoryCallback;
-
 	/** Enable all filters by default */
 	FOutputLogFilter() : TextFilterExpressionEvaluator(ETextFilterExpressionEvaluatorMode::BasicString)
 	{
@@ -278,10 +278,17 @@ struct FOutputLogFilter
 	/** Returns Evaluator syntax errors (if any) */
 	FText GetSyntaxErrors() const { return TextFilterExpressionEvaluator.GetFilterErrorText(); }
 
-	const TArray<FOutputLogCategorySettings>& GetCategoryFilters() const { return Categories; }
+	/** Returns the list of category settings sorted by Name. Ends up sorting the Categories Map if needed */
+	TArray<FOutputLogCategorySettings> GetCategoryFilters() const;
+
+	/** Returns the map of categories */
+	const TMap<FName, FOutputLogCategorySettings>& GetCategories() const
+	{
+		return Categories;
+	}
 
 	/** Adds a Log Category to the list of available categories, if it isn't already present */
-	void AddAvailableLogCategory(const FName& LogCategory, TOptional<bool> InitiallySelected = TOptional<bool>());
+	void AddAvailableLogCategory(const FName& LogCategory, const TOptional<bool>& InitiallySelected = TOptional<bool>());
 
 	/** Returns true if the specified log category is in the list of available categories */
 	bool IsLogCategoryAvailable(const FName& LogCategory) const;
@@ -301,15 +308,14 @@ struct FOutputLogFilter
 	void ClearSelectedLogCategories();
 	
 private:
-	FOutputLogCategorySettings* FindCategoryFilter(const FName& LogCategory);
-	const FOutputLogCategorySettings* FindCategoryFilter(const FName& LogCategory) const;
 	ELogLevelFilter GetMessageLevelFilter(const TSharedPtr<FOutputLogMessage>& Message) const;
-	
-	/** Expression evaluator that can be used to perform complex text filter queries */
-	FTextFilterExpressionEvaluator TextFilterExpressionEvaluator;
 
+	mutable bool bCategoriesNeedSorting = false;
 	/** All categories the filter is aware of */
-	TArray<FOutputLogCategorySettings> Categories;
+	TMap<FName, FOutputLogCategorySettings> Categories;
+
+	/** Expression evaluator that can be used to perform complex text filter queries */
+	FTextFilterExpressionEvaluator TextFilterExpressionEvaluator;
 };
 
 /**
@@ -329,7 +335,13 @@ public:
 		SLATE_EVENT(FSimpleDelegate, OnClearLog)
 
 		/** All messages captured before this log window has been created */
-		SLATE_ARGUMENT( TArray< TSharedPtr<FOutputLogMessage> >, Messages )
+		SLATE_ARGUMENT_DEFAULT( const TArray< TSharedPtr<FOutputLogMessage> >*, Messages) = nullptr;
+
+		/** The default Filter to apply to the OutputLog */
+		SLATE_ARGUMENT_DEFAULT( const FOutputLogFilter*, InitialLogFilter ) = nullptr;
+
+		/** The default FilterSettings to apply to the OutputLog */
+		SLATE_ARGUMENT_DEFAULT( const FOutputLogFilterSettings*, InitialLogFilterSettings ) = nullptr;
 
 		/**  */
 		SLATE_ARGUMENT_DEFAULT( EOutputLogSettingsMenuFlags, SettingsMenuFlags ) = EOutputLogSettingsMenuFlags::None;
@@ -434,9 +446,6 @@ protected:
 	/** The editable text showing all log messages */
 	TSharedPtr< SSearchBox > FilterTextBox;
 
-	/** True if the user has scrolled the window upwards */
-	bool bIsUserScrolled;
-
 	FSimpleDelegate OnClearLogDelegate; 
     FSimpleDelegate OnCloseConsole;
 
@@ -525,13 +534,16 @@ protected:
 
 	FDelegateHandle SettingsWatchHandle;
 
+	int32 LoggingLineLimit = 10000;
+
+	/** True if the user has scrolled the window upwards */
+	bool bIsUserScrolled = false;
+
 	bool bShouldCreateDrawerDockButton = false;
 
 	bool bShouldShowLoggingLimitMenu = false;
 
 	bool bEnableLoggingLimit = false;
-
-	int32 LoggingLineLimit = 10000;
 };
 
 /** Output log text marshaller to convert an array of FOutputLogMessages into styled lines to be consumed by an FTextLayout */
