Allow hiding 'Cmd' from output log when switching permissions that do not allow 'Cmd'
#rnx
#Synced-CL 52124361

diff --git a/Engine/Source/Developer/OutputLog/Private/OutputLogModule.cpp b/Engine/Source/Developer/OutputLog/Private/OutputLogModule.cpp
--- a/Engine/Source/Developer/OutputLog/Private/OutputLogModule.cpp
+++ b/Engine/Source/Developer/OutputLog/Private/OutputLogModule.cpp
@@ -184,6 +184,7 @@ void FOutputLogModule::StartupModule()
 
 	OutputLogHistory = MakeShareable(new FOutputLogHistory);
 	OutputLogFilterCache = MakeUnique<FOutputLogFilter>();
+	ConsoleCommandExecutorPermissionList = MakeUnique<FNamePermissionList>();
 
 	SOutputLog::RegisterSettingsMenu();
 }
@@ -546,3 +547,8 @@ bool FOutputLogModule::ShouldCycleToOutputLogDrawer() const
 	return true;
 #endif
 }
+
+const TUniquePtr<FNamePermissionList>& FOutputLogModule::GetConsoleCommandExecutorPermissionList() const
+{
+	return ConsoleCommandExecutorPermissionList;
+}
diff --git a/Engine/Source/Developer/OutputLog/Private/SOutputLog.cpp b/Engine/Source/Developer/OutputLog/Private/SOutputLog.cpp
--- a/Engine/Source/Developer/OutputLog/Private/SOutputLog.cpp
+++ b/Engine/Source/Developer/OutputLog/Private/SOutputLog.cpp
@@ -48,6 +48,32 @@
 
 #define LOCTEXT_NAMESPACE "SOutputLog"
 
+namespace OutputLogUtils
+{
+	TArray<IConsoleCommandExecutor*> GetAllowedConsoleCommandExecutors()
+	{
+		TArray<IConsoleCommandExecutor*> CommandExecutors = IModularFeatures::Get().GetModularFeatureImplementations<IConsoleCommandExecutor>(IConsoleCommandExecutor::ModularFeatureName());
+
+		const TUniquePtr<FNamePermissionList>& PermissionList = FOutputLogModule::Get().GetConsoleCommandExecutorPermissionList();
+		if (PermissionList->HasFiltering())
+		{
+			TArray<IConsoleCommandExecutor*> FilteredCommandExecutors;
+			FilteredCommandExecutors.Reserve(CommandExecutors.Num());
+			for (IConsoleCommandExecutor* CommandExecutor : CommandExecutors)
+			{
+				if (PermissionList->PassesFilter(CommandExecutor->GetName()))
+				{
+					FilteredCommandExecutors.Add(CommandExecutor);
+				}
+			}
+
+			return FilteredCommandExecutors;
+		}
+
+		return CommandExecutors;
+	}
+}
+
 class FCategoryLineHighlighter : public ISlateLineHighlighter
 {
 public:
@@ -178,6 +204,8 @@ void SConsoleInputBox::Construct(const FArguments& InArgs)
 
 	IModularFeatures::Get().OnModularFeatureRegistered().AddSP(this, &SConsoleInputBox::OnCommandExecutorRegistered);
 	IModularFeatures::Get().OnModularFeatureUnregistered().AddSP(this, &SConsoleInputBox::OnCommandExecutorUnregistered);
+	FOutputLogModule::Get().GetConsoleCommandExecutorPermissionList()->OnFilterChanged().AddSP(this, &SConsoleInputBox::SyncActiveCommandExecutor);
+
 	EPopupMethod PopupMethod = GIsEditor ? EPopupMethod::CreateNewWindow : EPopupMethod::UseCurrentWindow;
 	ChildSlot
 	[
@@ -639,7 +667,11 @@ void SConsoleInputBox::OnCommandExecutorUnregistered(const FName& Type, class IM
 
 void SConsoleInputBox::SyncActiveCommandExecutor()
 {
-	TArray<IConsoleCommandExecutor*> CommandExecutors = IModularFeatures::Get().GetModularFeatureImplementations<IConsoleCommandExecutor>(IConsoleCommandExecutor::ModularFeatureName());
+	TArray<IConsoleCommandExecutor*> CommandExecutors = OutputLogUtils::GetAllowedConsoleCommandExecutors();
+	CommandExecutors.Sort([](IConsoleCommandExecutor& LHS, IConsoleCommandExecutor& RHS)
+	{
+		return LHS.GetDisplayName().CompareTo(RHS.GetDisplayName()) < 0;
+	});
 	ActiveCommandExecutor = nullptr;
 
 	if (CommandExecutors.IsValidIndex(0))
@@ -716,7 +748,7 @@ bool SConsoleInputBox::IsCommandExecutorMenuEnabled() const
 void SConsoleInputBox::MakeNextCommandExecutorActive()
 {
 	// Sorted so the iteration order matches the displayed order.
-	TArray<IConsoleCommandExecutor*> CommandExecutors = IModularFeatures::Get().GetModularFeatureImplementations<IConsoleCommandExecutor>(IConsoleCommandExecutor::ModularFeatureName());
+	TArray<IConsoleCommandExecutor*> CommandExecutors = OutputLogUtils::GetAllowedConsoleCommandExecutors();
 	CommandExecutors.Sort([](IConsoleCommandExecutor& LHS, IConsoleCommandExecutor& RHS)
 		{
 			return LHS.GetDisplayName().CompareTo(RHS.GetDisplayName()) < 0;
@@ -749,7 +781,7 @@ TSharedRef<SWidget> SConsoleInputBox::GetCommandExecutorMenuContent()
 			{
 				if (TSharedPtr<SConsoleInputBox> This = Context->GetConsoleInputBox())
 				{
-					TArray<IConsoleCommandExecutor*> CommandExecutors = IModularFeatures::Get().GetModularFeatureImplementations<IConsoleCommandExecutor>(IConsoleCommandExecutor::ModularFeatureName());
+					TArray<IConsoleCommandExecutor*> CommandExecutors = OutputLogUtils::GetAllowedConsoleCommandExecutors();
 					CommandExecutors.Sort([](IConsoleCommandExecutor& LHS, IConsoleCommandExecutor& RHS)
 					{
 						return LHS.GetDisplayName().CompareTo(RHS.GetDisplayName()) < 0;
diff --git a/Engine/Source/Developer/OutputLog/Public/OutputLogModule.h b/Engine/Source/Developer/OutputLog/Public/OutputLogModule.h
--- a/Engine/Source/Developer/OutputLog/Public/OutputLogModule.h
+++ b/Engine/Source/Developer/OutputLog/Public/OutputLogModule.h
@@ -4,6 +4,7 @@
 
 #include "CoreMinimal.h"
 #include "Widgets/SWidget.h"
+#include "Misc/NamePermissionList.h"
 #include "Modules/ModuleInterface.h"
 #include "Widgets/SWindow.h"
 
@@ -107,6 +108,9 @@ public:
 
 	TMulticastDelegate<void()> OnOutputLogDrawerCleared;
 
+	/** Returns list used to filter which command executors are allowed */
+	virtual const TUniquePtr<FNamePermissionList>& GetConsoleCommandExecutorPermissionList() const;
+
 private:
 	TSharedRef<SDockTab> SpawnOutputLogTab(const FSpawnTabArgs& Args);
 	void SaveDockedTabSettings();
@@ -134,4 +138,7 @@ private:
 
 	/** Weak pointer to the widget to focus once they console window closes */
 	TWeakPtr<SWidget> PreviousKeyboardFocusedWidget;
+
+	/** List used to filter which command executors are allowed */
+	TUniquePtr<FNamePermissionList> ConsoleCommandExecutorPermissionList;
 };
