Add a right click menu optoin in the output log to open the file location in the file explorer, in addition to it being behind the settings cog, because it is more discoverable and easier to see

#rnx
#Synced-CL 52582364

diff --git a/Engine/Source/Developer/OutputLog/Private/OutputLogStyle.cpp b/Engine/Source/Developer/OutputLog/Private/OutputLogStyle.cpp
--- a/Engine/Source/Developer/OutputLog/Private/OutputLogStyle.cpp
+++ b/Engine/Source/Developer/OutputLog/Private/OutputLogStyle.cpp
@@ -71,6 +71,8 @@ FOutputLogStyle::FOutputLogStyle()
 		);
 
 		Set("DebugConsole.Icon", new IMAGE_BRUSH_SVG("Starship/Common/Console", Icon16x16));
+		Set("OutputLog.OpenSourceLocation", new IMAGE_BRUSH_SVG("Starship/Common/folder-open", Icon16x16));
+		Set("OutputLog.ClearLog", new IMAGE_BRUSH_SVG("Starship/Common/delete-outline", Icon16x16));
 	}
 
 	FSlateStyleRegistry::RegisterSlateStyle(*this);
diff --git a/Engine/Source/Developer/OutputLog/Private/SOutputLog.cpp b/Engine/Source/Developer/OutputLog/Private/SOutputLog.cpp
--- a/Engine/Source/Developer/OutputLog/Private/SOutputLog.cpp
+++ b/Engine/Source/Developer/OutputLog/Private/SOutputLog.cpp
@@ -1594,12 +1594,19 @@ void SOutputLog::ExtendTextBoxMenu(FMenuBuilder& Builder)
 		);
 
 	Builder.AddMenuEntry(
-		NSLOCTEXT("OutputLog", "ClearLogLabel", "Clear Log"), 
-		NSLOCTEXT("OutputLog", "ClearLogTooltip", "Clears all log messages"), 
-		FSlateIcon(), 
+		NSLOCTEXT("OutputLog", "ClearLogLabel", "Clear Log"),
+		NSLOCTEXT("OutputLog", "ClearLogTooltip", "Clears all log messages"),
+		FSlateIcon(FOutputLogStyle::Get().GetStyleSetName(), "OutputLog.ClearLog"),
 		ClearOutputLogAction
 		);
 
+	Builder.AddMenuEntry(
+		NSLOCTEXT("OutputLog", "ShowLogInExplorerLabel", "Show Log File in Explorer"),
+		NSLOCTEXT("OutputLog", "ShowLogInExplorerTooltip", "Opens the folder containing the current Output Log file in Explorer"),
+		FSlateIcon(FOutputLogStyle::Get().GetStyleSetName(), "OutputLog.OpenSourceLocation"),
+		FUIAction(FExecuteAction::CreateRaw(this, &SOutputLog::OpenLogFileInExplorer))
+	);
+
 	Builder.AddMenuEntry(
 		FUIAction(
 			FExecuteAction::CreateLambda([this](){ bEnableLoggingLimit = !bEnableLoggingLimit; }),
