[nDisplay] nDisplay cluster monitor

Editor panel that displays all live nDisplay clusters and their observable entities. Each entity can be streamed to the editor via NDI media and previewed here. Additional tools, such as remote console command execution, are also available.

[RN] Implemented nDisplay cluster monitor for remote observation of cluster nodes
#Synced-CL 0

diff --git a/Engine/Plugins/Runtime/nDisplay/Source/DisplayCluster/Private/Game/EngineClasses/Basics/DisplayClusterViewportClient.cpp b/Engine/Plugins/Runtime/nDisplay/Source/DisplayCluster/Private/Game/EngineClasses/Basics/DisplayClusterViewportClient.cpp
--- a/Engine/Plugins/Runtime/nDisplay/Source/DisplayCluster/Private/Game/EngineClasses/Basics/DisplayClusterViewportClient.cpp
+++ b/Engine/Plugins/Runtime/nDisplay/Source/DisplayCluster/Private/Game/EngineClasses/Basics/DisplayClusterViewportClient.cpp
@@ -355,7 +355,7 @@ void UDisplayClusterViewportClient::Draw(FViewport* InViewport, FCanvas* SceneCa
 		DebugCanvas->SetScaledToRenderTarget(bStereoRendering);
 		DebugCanvas->SetStereoRendering(bStereoRendering);
 
-		const bool bGuiPropagationActive = FDisplayClusterGuiLayerController::Get().IsActiveThisFrame();
+		const bool bGuiPropagationActive = FDisplayClusterGuiLayerController::Get().IsEnabled();
 		DebugCanvas->SetWriteDestinationAlpha(bGuiPropagationActive);
 	}
 	if (SceneCanvas)
diff --git a/Engine/Plugins/Runtime/nDisplay/Source/DisplayCluster/Private/Misc/DisplayClusterConsoleExec.cpp b/Engine/Plugins/Runtime/nDisplay/Source/DisplayCluster/Private/Misc/DisplayClusterConsoleExec.cpp
--- a/Engine/Plugins/Runtime/nDisplay/Source/DisplayCluster/Private/Misc/DisplayClusterConsoleExec.cpp
+++ b/Engine/Plugins/Runtime/nDisplay/Source/DisplayCluster/Private/Misc/DisplayClusterConsoleExec.cpp
@@ -1,53 +1,75 @@
 // Copyright Epic Games, Inc. All Rights Reserved.
 
-#include "DisplayClusterConsoleExec.h"
-#include "Misc/DisplayClusterLog.h"
-#include "Cluster/DisplayClusterClusterEvent.h"
+#include "Misc/DisplayClusterConsoleExec.h"
 
+#include "Cluster/DisplayClusterClusterEvent.h"
 #include "Engine/Engine.h"
 #include "Engine/LocalPlayer.h"
 #include "Features/IModularFeatures.h"
+#include "Misc/DisplayClusterLog.h"
 
 
 bool FDisplayClusterConsoleExec::Exec(const FDisplayClusterClusterEventJson& InEvent)
 {
+	// Extract command
 	const FString* const ExecString = InEvent.Parameters.Find("ExecString");
-	if (!ExecString || !ExecString->Len())
+	if (!ExecString || ExecString->IsEmpty())
 	{
 		UE_LOGF(LogDisplayClusterModule, Warning, "ConsoleExec ignoring cluster event with no ExecString");
 		return false;
 	}
 
+	// Extract executor name
 	const FString* const RequestedExecutor = InEvent.Parameters.Find("Executor");
 
-	UE_LOGF(LogDisplayClusterModule, Log, "ConsoleExec cluster event: Executor='%ls', ExecString='%ls'", RequestedExecutor ? **RequestedExecutor : TEXT(""), **ExecString);
+	UE_LOGF(LogDisplayClusterModule, Verbose, "ConsoleExec cluster event: Executor='%ls', ExecString='%ls'", RequestedExecutor ? **RequestedExecutor : TEXT(""), **ExecString);
+
+	// Execute
+	return Exec(RequestedExecutor ? **RequestedExecutor : FString(), *ExecString);
+}
+
+bool FDisplayClusterConsoleExec::Exec(const FString& InExecutorName, const FString& InExecCommand)
+{
+	if (InExecCommand.IsEmpty())
+	{
+		return false;
+	}
+
+	UE_LOGF(LogDisplayClusterModule, Log, "ConsoleExec command: Executor='%ls', ExecString='%ls'", *InExecutorName, *InExecCommand);
 
-	if (RequestedExecutor && RequestedExecutor->Len())
+	if (!InExecutorName.IsEmpty())
 	{
-		TArray<IConsoleCommandExecutor*> CommandExecutors = IModularFeatures::Get().GetModularFeatureImplementations<IConsoleCommandExecutor>(IConsoleCommandExecutor::ModularFeatureName());
+		const FName FeatureName = IConsoleCommandExecutor::ModularFeatureName();
+		TArray<IConsoleCommandExecutor*> CommandExecutors = IModularFeatures::Get().GetModularFeatureImplementations<IConsoleCommandExecutor>(FeatureName);
+
+		const FName ExecutorName = *InExecutorName;
 		for (IConsoleCommandExecutor* CommandExecutor : CommandExecutors)
 		{
-			if (CommandExecutor->GetName() == FName(*RequestedExecutor))
+			if (CommandExecutor->GetName() == ExecutorName)
 			{
-				return CommandExecutor->Exec(**ExecString);
+				return CommandExecutor->Exec(*InExecCommand);
 			}
 		}
 
-		UE_LOGF(LogDisplayClusterModule, Warning, "ConsoleExec couldn't find requested executor: %ls", **RequestedExecutor);
+		UE_LOGF(LogDisplayClusterModule, Warning, "ConsoleExec couldn't find requested executor: %ls", *InExecutorName);
 		return false;
 	}
 	else
 	{
 		// With no explicit executor, try to route this as an Unreal console
 		// command as best we can. Available context is limited here.
-		ULocalPlayer* const Player = GEngine->GetDebugLocalPlayer();
-		if (Player)
-		{
-			return Player->Exec(Player->GetWorld(), **ExecString, *GLog);
-		}
-		else
+		if (GEngine)
 		{
-			return GEngine->Exec(nullptr, **ExecString, *GLog);
+			if (ULocalPlayer* const Player = GEngine->GetDebugLocalPlayer())
+			{
+				return Player->Exec(Player->GetWorld(), *InExecCommand, *GLog);
+			}
+			else
+			{
+				return GEngine->Exec(nullptr, *InExecCommand, *GLog);
+			}
 		}
 	}
+
+	return false;
 }
diff --git a/Engine/Plugins/Runtime/nDisplay/Source/DisplayCluster/Private/Render/DisplayClusterRenderManager.cpp b/Engine/Plugins/Runtime/nDisplay/Source/DisplayCluster/Private/Render/DisplayClusterRenderManager.cpp
--- a/Engine/Plugins/Runtime/nDisplay/Source/DisplayCluster/Private/Render/DisplayClusterRenderManager.cpp
+++ b/Engine/Plugins/Runtime/nDisplay/Source/DisplayCluster/Private/Render/DisplayClusterRenderManager.cpp
@@ -1,51 +1,48 @@
 // Copyright Epic Games, Inc. All Rights Reserved.
 
 #include "Render/DisplayClusterRenderManager.h"
+
+#include "Components/DisplayClusterCameraComponent.h"
 #include "Config/IPDisplayClusterConfigManager.h"
 
-#include "Engine/GameViewportClient.h"
 #include "Engine/GameEngine.h"
+#include "Engine/GameViewportClient.h"
+#include "Engine/Scene.h"
+
+#include "Framework/Application/SlateApplication.h"
+#include "Game/IPDisplayClusterGameManager.h"
+#include "Kismet/GameplayStatics.h"
 
+#include "Misc/DisplayClusterDataCache.h"
 #include "Misc/DisplayClusterGlobals.h"
 #include "Misc/DisplayClusterHelpers.h"
 #include "Misc/DisplayClusterLog.h"
 #include "Misc/DisplayClusterStrings.h"
 
-#include "DisplayClusterConfigurationStrings.h"
-
-#include "Components/DisplayClusterCameraComponent.h"
-#include "DisplayClusterViewportClient.h"
-#include "DisplayClusterConfigurationTypes.h"
-
-#include "Game/IPDisplayClusterGameManager.h"
-
+#include "Render/Containers/DisplayClusterRender_MeshComponent.h"
+#include "Render/Containers/DisplayClusterRender_Texture.h"
 #include "Render/Device/DisplayClusterRenderDeviceFactoryInternal.h"
-
+#include "Render/Device/IDisplayClusterRenderDevice.h"
 #include "Render/Device/IDisplayClusterRenderDeviceFactory.h"
+#include "Render/GUILayer/DisplayClusterGuiLayerController.h"
 #include "Render/Monitoring/DisplayClusterVblankMonitor.h"
 #include "Render/PostProcess/IDisplayClusterPostProcess.h"
-#include "Render/Projection/IDisplayClusterProjectionPolicyFactory.h"
-#include "Render/Projection/IDisplayClusterProjectionPolicy.h"
-#include "Render/Synchronization/IDisplayClusterRenderSyncPolicyFactory.h"
-
-#include "Render/Containers/DisplayClusterRender_MeshComponent.h"
-#include "Render/Containers/DisplayClusterRender_Texture.h"
-
 #include "Render/Presentation/DisplayClusterPresentationNative.h"
-
+#include "Render/Projection/IDisplayClusterProjectionPolicy.h"
+#include "Render/Projection/IDisplayClusterProjectionPolicyFactory.h"
 #include "Render/Synchronization/DisplayClusterRenderSyncPolicyFactoryInternal.h"
 #include "Render/Synchronization/DisplayClusterRenderSyncPolicyNone.h"
+#include "Render/Synchronization/IDisplayClusterRenderSyncPolicyFactory.h"
 
-#include "Framework/Application/SlateApplication.h"
 #include "Slate/SlateViewportProvider.h"
 
-#include "UnrealClient.h"
-#include "Kismet/GameplayStatics.h"
-
 #include "CineCameraComponent.h"
-#include "Engine/Scene.h"
-
+#include "DisplayClusterConfigurationStrings.h"
+#include "DisplayClusterConfigurationTypes.h"
 #include "DisplayClusterRootActor.h"
+#include "DisplayClusterViewportClient.h"
+#include "UnrealClient.h"
+
 
 #if PLATFORM_WINDOWS
 #include "Windows/AllowWindowsPlatformTypes.h"
@@ -64,7 +61,6 @@ static TAutoConsoleVariable<int32> CVarSyncDiagnosticsVBlankMonitoring(
 	ECVF_ReadOnly
 );
 
-#include "Misc/DisplayClusterDataCache.h"
 
 /**
  * The cache for DC texture objects. (Singleton)
@@ -617,6 +613,11 @@ TSharedPtr<IDisplayClusterRender_Texture, ESPMode::ThreadSafe> FDisplayClusterRe
 	return FDisplayClusterRenderTextureCache::GetOrCreateRenderTexture(InUniqueTextureName);
 }
 
+IDisplayClusterGUILayerController& FDisplayClusterRenderManager::GetGuiLayerController() const
+{
+	return FDisplayClusterGuiLayerController::Get();
+}
+
 IDisplayClusterViewportManager* FDisplayClusterRenderManager::GetViewportManager() const
 {
 	ADisplayClusterRootActor* RootActor = GDisplayCluster->GetGameMgr()->GetRootActor();
diff --git a/Engine/Plugins/Runtime/nDisplay/Source/DisplayCluster/Private/Render/DisplayClusterRenderManager.h b/Engine/Plugins/Runtime/nDisplay/Source/DisplayCluster/Private/Render/DisplayClusterRenderManager.h
--- a/Engine/Plugins/Runtime/nDisplay/Source/DisplayCluster/Private/Render/DisplayClusterRenderManager.h
+++ b/Engine/Plugins/Runtime/nDisplay/Source/DisplayCluster/Private/Render/DisplayClusterRenderManager.h
@@ -6,7 +6,7 @@
 #include "RHIFwd.h"
 #include "Render/IPDisplayClusterRenderManager.h"
 
-
+class IDisplayClusterGUILayerController;
 class IDisplayClusterPostProcess;
 class IDisplayClusterPostProcessFactory; 
 class IDisplayClusterProjectionPolicy;
@@ -70,6 +70,7 @@ public:
 	// Resources
 	virtual TSharedPtr<IDisplayClusterRender_MeshComponent, ESPMode::ThreadSafe> CreateMeshComponent() const override;
 	virtual TSharedPtr<IDisplayClusterRender_Texture, ESPMode::ThreadSafe> GetOrCreateCachedTexture(const FString& InUniqueTextureName) const override;
+	virtual IDisplayClusterGUILayerController& GetGuiLayerController() const override;
 
 	virtual TSharedPtr<IDisplayClusterVblankMonitor, ESPMode::ThreadSafe> GetVblankMonitor() override
 	{
diff --git a/Engine/Plugins/Runtime/nDisplay/Source/DisplayCluster/Private/Render/GUILayer/DisplayClusterGuiLayerController.cpp b/Engine/Plugins/Runtime/nDisplay/Source/DisplayCluster/Private/Render/GUILayer/DisplayClusterGuiLayerController.cpp
--- a/Engine/Plugins/Runtime/nDisplay/Source/DisplayCluster/Private/Render/GUILayer/DisplayClusterGuiLayerController.cpp
+++ b/Engine/Plugins/Runtime/nDisplay/Source/DisplayCluster/Private/Render/GUILayer/DisplayClusterGuiLayerController.cpp
@@ -18,12 +18,14 @@
 #include "Slate/SceneViewport.h"
 
 
-TAutoConsoleVariable<bool> CVarPropagateGui(
+TAutoConsoleVariable<int32> CVarPropagateGui(
 	TEXT("nDisplay.GUI.Propagate"),
-	false,
+	1,
 	TEXT("Show GUI on viewports\n")
-	TEXT("0 : Disabled\n")
-	TEXT("1 : Enabled\n"),
+	TEXT("0 : Hidden     - never rendered to the backbuffer\n")
+	TEXT("1 : Whole      - GUI is rendered in backbuffer space\n")
+	TEXT("2 : Propagated - GUI is rendered to nDisplay viewports\n")
+	TEXT("Other values are treated as '0'.\n"),
 	ECVF_RenderThreadSafe
 );
 
@@ -47,13 +49,23 @@ FDisplayClusterGuiLayerController& FDisplayClusterGuiLayerController::Get()
 	return Instance;
 }
 
+FIntPoint FDisplayClusterGuiLayerController::GetGuiLayerTextureSize()
+{
+	if (bEnabled && GEngine && GEngine->GameViewport && GEngine->GameViewport->Viewport)
+	{
+		return GEngine->GameViewport->Viewport->GetSizeXY();
+	}
+
+	return FIntPoint::ZeroValue;
+}
+
 FRDGTextureRef FDisplayClusterGuiLayerController::ProcessFinalTexture_RenderThread(FRDGBuilder& GraphBuilder, FRDGTextureRef FinalTexture)
 {
 	checkSlow(IsInRenderingThread());
 	checkSlow(FinalTexture);
 
-	// Nothing to do if GUI layer propagation is not active
-	if (!bActiveThisFrameRT || !bEnabled)
+	// Nothing to do if not active
+	if (!bEnabled)
 	{
 		return FinalTexture;
 	}
@@ -71,6 +83,33 @@ FRDGTextureRef FDisplayClusterGuiLayerController::ProcessFinalTexture_RenderThre
 		return FinalTexture;
 	}
 
+	// Conditionally draw the GUI on top of the scene view
+	if (OutputMode == EGuiOutputMode::Whole)
+	{
+		if (!TexViewportGUILayerRHI.IsValid())
+		{
+			return FinalTexture;
+		}
+
+		const FIntPoint SceneTextureSize = Output->Desc.Extent;
+
+		// Prepare textures for drawing
+		FRDGTextureRef GUITexture  = RegisterExternalTexture(GraphBuilder, TexViewportGUILayerRHI, TEXT("nD.TexViewportGUILayerRHI"));
+		FRDGTextureRef TempTexture = CreateTextureFrom_RenderThread(GraphBuilder, Output, TEXT("nD.TexTempSceneRDG"), Output->Desc.Extent);
+
+		// Make a copy to the input texture
+		AddCopyTexturePass(GraphBuilder, Output, TempTexture, FRHICopyTextureInfo{});
+
+		// Draw the GUI layer on top of the nD viewport
+		IDisplayClusterShaders::Get().AddDrawOverlayPass(GraphBuilder,
+			FDisplayClusterShaderParameters_Overlay
+			{
+				.BaseTexture    = TempTexture, // Scene color (temp copy)
+				.OverlayTexture = GUITexture,  // GUI layer
+				.OutputTexture  = Output,      // Scene color (original viewport)
+			});
+	}
+
 	// If everything Ok, return the original viewport's buffer
 	return Output;
 }
@@ -80,8 +119,13 @@ void FDisplayClusterGuiLayerController::ProcessPostRenderViewFamily_RenderThread
 	checkSlow(IsInRenderingThread());
 	checkSlow(ViewportProxy);
 
+	if (!bEnabled)
+	{
+		return;
+	}
+
 	// Nothing to do if GUI layer propagation is not active
-	if (!bActiveThisFrameRT || !bEnabled)
+	if (OutputModeRT != EGuiOutputMode::Propagated)
 	{
 		return;
 	}
@@ -124,6 +168,19 @@ void FDisplayClusterGuiLayerController::ProcessPostRenderViewFamily_RenderThread
 	}
 }
 
+FDisplayClusterGuiLayerController::EGuiOutputMode FDisplayClusterGuiLayerController::GetOutputMode() const
+{
+	const int32 Value = CVarPropagateGui.GetValueOnAnyThread();
+
+	switch (Value)
+	{
+	case 0:  return EGuiOutputMode::Hidden;
+	case 1:  return EGuiOutputMode::Whole;
+	case 2:  return EGuiOutputMode::Propagated;
+	default: return EGuiOutputMode::Hidden;
+	}
+}
+
 void FDisplayClusterGuiLayerController::HandleSlatePreTick(float InDeltaTime)
 {
 	checkSlow(IsInGameThread());
@@ -135,17 +192,13 @@ void FDisplayClusterGuiLayerController::HandleSlatePreTick(float InDeltaTime)
 	}
 
 	// Read the cvar once per frame
-	bActiveThisFrame = CVarPropagateGui.GetValueOnGameThread();
+	OutputMode = GetOutputMode();
 
 	ENQUEUE_RENDER_COMMAND(FlushCanvasRT)(
-		[this, bNewActiveThisFrame = bActiveThisFrame](FRHICommandListImmediate& RHICmdList)
+		[this, NewOutputMode = OutputMode](FRHICommandListImmediate& RHICmdList)
 		{
-			// Update the corresponding flag, and leave if inactive
-			bActiveThisFrameRT = bNewActiveThisFrame;
-			if (!bActiveThisFrameRT)
-			{
-				return;
-			}
+			// Update the corresponding flag
+			OutputModeRT = NewOutputMode;
 
 			// Get game viewport's current render target
 			FTextureRHIRef ViewportRTT = GetViewportRTT_RenderThread();
@@ -161,7 +214,7 @@ void FDisplayClusterGuiLayerController::HandleSlatePreTick(float InDeltaTime)
 				return;
 			}
 
-			// Store the original vieport's buffer, and substitute is with our float16 buffer
+			// Store the original viewport's buffer, and substitute is with our float16 buffer
 			TexViewportOriginalRHI = ViewportRTT;
 			SetViewportRTT_RenderThread(TexViewportGUILayerRHI);
 
diff --git a/Engine/Plugins/Runtime/nDisplay/Source/DisplayCluster/Private/Render/GUILayer/DisplayClusterGuiLayerController.h b/Engine/Plugins/Runtime/nDisplay/Source/DisplayCluster/Private/Render/GUILayer/DisplayClusterGuiLayerController.h
--- a/Engine/Plugins/Runtime/nDisplay/Source/DisplayCluster/Private/Render/GUILayer/DisplayClusterGuiLayerController.h
+++ b/Engine/Plugins/Runtime/nDisplay/Source/DisplayCluster/Private/Render/GUILayer/DisplayClusterGuiLayerController.h
@@ -2,11 +2,12 @@
 
 #pragma once
 
+#include "Render/GUILayer/IDisplayClusterGUILayerController.h"
+
 #include "PixelFormat.h"
 #include "RHIResources.h"
 #include "UnrealClient.h"
 
-
 class FRDGBuilder;
 class FSceneViewFamily;
 class IDisplayClusterViewportProxy;
@@ -15,38 +16,46 @@ class IDisplayClusterViewportProxy;
 /**
  * GUI layer controller singleton
  * 
- * This class is responsible for propagation of the GUI layer (includes canvases + UMG) across
- * all nDisplay viewports.
+ * This class is responsible for propagation of the GUI layer
+ * (includes canvases + UMG) across all nDisplay viewports.
+ * It also exposes the GUI layer texture to external users.
  */
 class FDisplayClusterGuiLayerController
+	: public IDisplayClusterGUILayerController
 {
 protected:
+
 	FDisplayClusterGuiLayerController();
 	~FDisplayClusterGuiLayerController() = default;
 
 public:
+
+	/** Singleton access */
 	static FDisplayClusterGuiLayerController& Get();
 
 public:
 
-	/**
-	 * Returns true if this controller is enabled
-	 */
-	bool IsEnabled() const
+	//~ Begin IDisplayClusterGUILayerController interface
+
+	/** Returns true if GUI controller is enabled during current session */
+	virtual bool IsEnabled() const override
 	{
 		return bEnabled;
 	}
 
-	/**
-	 * Returns true if this controller is active this frame
-	 */
-	bool IsActiveThisFrame() const
+	/** Returns current GUI layer texture size */
+	virtual FIntPoint GetGuiLayerTextureSize() override;
+
+	/** Returns current GUI layer texture if available, otherwise invalid reference */
+	virtual FTextureRHIRef GetGuiLayerTexture_RenderThread() override
 	{
-		checkSlow(IsInGameThread() || IsInRenderingThread());
-		const bool bActive = (IsInRenderingThread() ? bActiveThisFrameRT : bActiveThisFrame);
-		return IsEnabled() && bActive;
+		return TexViewportGUILayerRHI;
 	}
 
+	//~ End IDisplayClusterGUILayerController interface
+
+public:
+
 	/**
 	 * Called before copying the final output to the backbuffer. It's used to substitute the current
 	 * viewport's buffer to its original one stored on SlatePreTick (refer corresponding method).
@@ -64,6 +73,21 @@ public:
 	 */
 	void ProcessPostRenderViewFamily_RenderThread(FRDGBuilder& GraphBuilder, const FSceneViewFamily& ViewFamiliy, const IDisplayClusterViewportProxy* ViewportProxy);
 
+private:
+
+	/**
+	 * Supported GUI output modes
+	 */
+	enum class EGuiOutputMode : uint8
+	{
+		Hidden     = 0, // Never render to the output
+		Whole      = 1, // Render in the output (backbuffer) space like UE natively does
+		Propagated = 2, // Per-viewport GUI output
+	};
+
+	/** Converts current output mode cvar value into a corresponding enum */
+	EGuiOutputMode GetOutputMode() const;
+
 private:
 
 	/**
@@ -74,7 +98,7 @@ private:
 
 private:
 
-	/** Returns current game vieport buffer */
+	/** Returns current game viewport buffer */
 	FTextureRHIRef GetViewportRTT_RenderThread();
 
 	/** Sets up new game viewport bufer (stays there until end of the frame) */
@@ -109,11 +133,11 @@ private:
 	/** Whether GUI layer propagation is active within current session */
 	const bool bEnabled;
 
-	/** Whether GUI layer propagation is active within current frame (game thread) */
-	bool bActiveThisFrame = false;
+	/** Current GUI output mode (game thread) */
+	EGuiOutputMode OutputMode = EGuiOutputMode::Whole;
 
-	/** Whether GUI layer propagation is active within current frame (render thread) */
-	bool bActiveThisFrameRT = false;
+	/** Current GUI output mode (render thread) */
+	EGuiOutputMode OutputModeRT = EGuiOutputMode::Whole;
 
 	/** The RHI texture reference to the original game viewport buffer  */
 	FTextureRHIRef TexViewportOriginalRHI;
diff --git a/Engine/Plugins/Runtime/nDisplay/Source/DisplayCluster/Public/Render/GUILayer/IDisplayClusterGUILayerController.h b/Engine/Plugins/Runtime/nDisplay/Source/DisplayCluster/Public/Render/GUILayer/IDisplayClusterGUILayerController.h
new file mode 100644
--- /dev/null
+++ b/Engine/Plugins/Runtime/nDisplay/Source/DisplayCluster/Public/Render/GUILayer/IDisplayClusterGUILayerController.h
@@ -0,0 +1,37 @@
+// Copyright Epic Games, Inc. All Rights Reserved.
+
+#pragma once
+
+#include "RHIFwd.h"
+#include "Math/MathFwd.h"
+
+
+/**
+ * GUI controller interface
+ * 
+ * Represents a mediator that controls how Slate/UMG output is rendered
+ * to the backbuffer. It also provides access to the texture containing
+ * the GUI (the so-called GUI layer).
+ */
+class IDisplayClusterGUILayerController
+{
+public:
+	virtual ~IDisplayClusterGUILayerController() = default;
+
+public:
+
+	/** Returns true if GUI controller is enabled during current session */
+	virtual bool IsEnabled() const = 0;
+
+	UE_DEPRECATED(5.8, "This function is deprecated. Use IsEnabled. ")
+	virtual bool IsActiveThisFrame() const
+	{
+		return IsEnabled();
+	}
+
+	/** Returns current GUI layer texture size */
+	virtual FIntPoint GetGuiLayerTextureSize() = 0;
+
+	/** Returns current GUI layer texture if available, otherwise invalid reference */
+	virtual FTextureRHIRef GetGuiLayerTexture_RenderThread() = 0;
+};
diff --git a/Engine/Plugins/Runtime/nDisplay/Source/DisplayCluster/Public/Render/IDisplayClusterRenderManager.h b/Engine/Plugins/Runtime/nDisplay/Source/DisplayCluster/Public/Render/IDisplayClusterRenderManager.h
--- a/Engine/Plugins/Runtime/nDisplay/Source/DisplayCluster/Public/Render/IDisplayClusterRenderManager.h
+++ b/Engine/Plugins/Runtime/nDisplay/Source/DisplayCluster/Public/Render/IDisplayClusterRenderManager.h
@@ -2,8 +2,9 @@
 
 #pragma once
 
-#include "Render/Device/IDisplayClusterRenderDevice.h"
+#include "CoreMinimal.h"
 
+class IDisplayClusterGUILayerController;
 class IDisplayClusterRenderDevice;
 class IDisplayClusterPostProcess;
 class IDisplayClusterPostProcessFactory;
@@ -219,4 +220,9 @@ public:
 	 * Get or create a cached texture with unique name
 	 */
 	virtual TSharedPtr<IDisplayClusterRender_Texture, ESPMode::ThreadSafe> GetOrCreateCachedTexture(const FString& InUniqueTextureName) const = 0;
+
+	/**
+	 * Returns GUI layer controller interface
+	 */
+	virtual IDisplayClusterGUILayerController& GetGuiLayerController() const = 0;
 };
diff --git a/Engine/Plugins/Runtime/nDisplay/Source/DisplayCluster/Public/Render/Viewport/RenderFrame/DisplayClusterRenderFrame.h b/Engine/Plugins/Runtime/nDisplay/Source/DisplayCluster/Public/Render/Viewport/RenderFrame/DisplayClusterRenderFrame.h
--- a/Engine/Plugins/Runtime/nDisplay/Source/DisplayCluster/Public/Render/Viewport/RenderFrame/DisplayClusterRenderFrame.h
+++ b/Engine/Plugins/Runtime/nDisplay/Source/DisplayCluster/Public/Render/Viewport/RenderFrame/DisplayClusterRenderFrame.h
@@ -52,7 +52,7 @@ struct FDisplayClusterRenderFrameTargetViewFamily
 	// Extensions that can modify view parameters
 	TArray<FSceneViewExtensionRef> ViewExtensions;
 
-	// Vieports, rendered at once for tthis family
+	// Viewports, rendered at once for this family
 	TArray<FDisplayClusterRenderFrameTargetView> Views;
 };
 
diff --git a/Engine/Plugins/Runtime/nDisplay/nDisplay.uplugin b/Engine/Plugins/Runtime/nDisplay/nDisplay.uplugin
--- a/Engine/Plugins/Runtime/nDisplay/nDisplay.uplugin
+++ b/Engine/Plugins/Runtime/nDisplay/nDisplay.uplugin
@@ -188,6 +188,24 @@
 				"Linux"
 			]
 		},
+		{
+			"Name": "DisplayClusterMonitor",
+			"Type": "Runtime",
+			"LoadingPhase": "Default",
+			"PlatformAllowList": [
+				"Win64",
+				"Linux"
+			]
+		},
+		{
+			"Name": "DisplayClusterMonitorEditor",
+			"Type": "Editor",
+			"LoadingPhase": "Default",
+			"PlatformAllowList": [
+				"Win64",
+				"Linux"
+			]
+		},
 		{
 			"Name": "DisplayClusterMultiUser",
 			"Type": "UncookedOnly",
@@ -272,6 +290,10 @@
 			"Name": "MediaIOFramework",
 			"Enabled": true
 		},
+		{
+			"Name": "NDIMedia",
+			"Enabled": true
+		},
 		{
 			"Name": "nDisplayModularFeatures",
 			"Enabled": true
@@ -288,6 +310,10 @@
 			"Name": "RemoteControlInterception",
 			"Enabled": true
 		},
+		{
+			"Name": "MediaPlayerEditor",
+			"Enabled": true
+		},
 		{
 			"Name": "MovieRenderPipeline",
 			"Enabled": true
diff --git a/Engine/Source/Developer/OutputLog/Private/OutputLogModule.cpp b/Engine/Source/Developer/OutputLog/Private/OutputLogModule.cpp
--- a/Engine/Source/Developer/OutputLog/Private/OutputLogModule.cpp
+++ b/Engine/Source/Developer/OutputLog/Private/OutputLogModule.cpp
@@ -235,6 +235,26 @@ TSharedRef<SWidget> FOutputLogModule::MakeConsoleInputBox(TSharedPtr<SMultiLineE
 	return NewConsoleInputBox;
 }
 
+UE::OutputLog::FConsoleInputBoxCreationResult FOutputLogModule::MakeConsoleInputBox(const UE::OutputLog::FConsoleInputBoxCreationParams& InCreationParams) const
+{
+	// Create requested widget
+	TSharedRef<SConsoleInputBox> NewConsoleInputBox = SNew(SConsoleInputBox)
+		.Visibility(MakeAttributeLambda([]() { return FOutputLogModule::Get().ShouldHideConsole() ? EVisibility::Collapsed : EVisibility::Visible; }))
+		.OnCloseConsole(InCreationParams.OnCloseConsole)
+		.OnConsoleCommandExecuted(InCreationParams.OnConsoleCommandExecuted)
+		.OnInterceptConsoleCommand(InCreationParams.OnInterceptConsoleCommand)
+		.OnHistoryOverride(InCreationParams.OnHistoryOverride);
+
+	// Prepare the results
+	const UE::OutputLog::FConsoleInputBoxCreationResult Result
+	{
+		.ConsoleInputBox        = NewConsoleInputBox,
+		.ExposedEditableTextBox = NewConsoleInputBox->GetEditableTextBox(),
+	};
+
+	return Result;
+}
+
 TSharedRef<SWidget> FOutputLogModule::MakeOutputLogDrawerWidget(const FSimpleDelegate& OnCloseConsole)
 {
 	TSharedPtr<SOutputLog> OutputLogDrawerPinned = OutputLogDrawer.Pin();
diff --git a/Engine/Source/Developer/OutputLog/Private/SOutputLog.cpp b/Engine/Source/Developer/OutputLog/Private/SOutputLog.cpp
--- a/Engine/Source/Developer/OutputLog/Private/SOutputLog.cpp
+++ b/Engine/Source/Developer/OutputLog/Private/SOutputLog.cpp
@@ -189,7 +189,9 @@ void SConsoleInputBox::Construct(const FArguments& InArgs)
 {
 	OnConsoleCommandExecuted = InArgs._OnConsoleCommandExecuted;
 	ConsoleCommandCustomExec = InArgs._ConsoleCommandCustomExec;
-	OnCloseConsole = InArgs._OnCloseConsole;
+	InterceptConsoleCommand  = InArgs._OnInterceptConsoleCommand;
+	HistoryOverride          = InArgs._OnHistoryOverride;
+	OnCloseConsole           = InArgs._OnCloseConsole;
 
 	if (!ConsoleCommandCustomExec.IsBound()) // custom execs always show the default executor in the UI (which has the selector disabled)
 	{
@@ -440,21 +442,34 @@ void SConsoleInputBox::OnTextCommitted( const FText& InText, ETextCommit::Type C
 			// re-entered.  We want to make sure the text string is empty on re-entry, so we'll clear it out
 			const FString ExecString = InText.ToString();
 
+			bool bWasIntercepted = false;
+			if (InterceptConsoleCommand.IsBound())
+			{
+				// Get the executor name (NAME_None represents the default console executor)
+				const FName ExecutorName = ActiveCommandExecutor ? ActiveCommandExecutor->GetName() : NAME_None;
+				// Execute interception logic; returns whether the command was intercepted.
+				bWasIntercepted = InterceptConsoleCommand.Execute(ExecutorName, ExecString);
+			}
+
 			// Clear the console input area
 			bIgnoreUIUpdate = true;
 			InputText->SetText(FText::GetEmpty());
 			ClearSuggestions();
 			bIgnoreUIUpdate = false;
-			
-			// Exec!
-			if (ConsoleCommandCustomExec.IsBound())
-			{
-				IConsoleManager::Get().AddConsoleHistoryEntry(TEXT(""), *ExecString);
-				ConsoleCommandCustomExec.Execute(ExecString);
-			}
-			else if (ActiveCommandExecutor)
+
+			// If not intercepted, then execute it
+			if (!bWasIntercepted)
 			{
-				ActiveCommandExecutor->Exec(*ExecString);
+				// Exec!
+				if (ConsoleCommandCustomExec.IsBound())
+				{
+					IConsoleManager::Get().AddConsoleHistoryEntry(TEXT(""), *ExecString);
+					ConsoleCommandCustomExec.Execute(ExecString);
+				}
+				else if (ActiveCommandExecutor)
+				{
+					ActiveCommandExecutor->Exec(*ExecString);
+				}
 			}
 		}
 		else
@@ -512,29 +527,10 @@ FReply SConsoleInputBox::OnPreviewKeyDown(const FGeometry& MyGeometry, const FKe
 			const bool bShowHistory = InputText->GetText().IsEmpty() || KeyEvent.IsControlDown();
 			if (bShowHistory)
 			{
-				IConsoleManager& ConsoleManager = IConsoleManager::Get();
-				TArray<FString> HistoryNames;
-				if (ActiveCommandExecutor)
-				{
-					ActiveCommandExecutor->GetExecHistory(HistoryNames);
-				}
-				else
-				{
-					ConsoleManager.GetConsoleHistory(TEXT(""), HistoryNames);
-				}
-				TArray<FConsoleSuggestion> History;
-				for (const FString& Name : HistoryNames)
-				{
-					FString HelpString;
-					// Try to find a console object for this history entry in order to retrieve a help string if possible :
-					const TCHAR* NamePtr = *Name;
-					if (IConsoleObject* CObj = ConsoleManager.FindConsoleObject(*FParse::Token(NamePtr, /*UseEscape = */false), /*bTrackFrequentCalls = */false); CObj && CObj->IsEnabled())
-					{
-						HelpString = CObj->GetDetailedHelp().ToString();
-					}
+				// Generate suggestions list
+				TArray<FConsoleSuggestion> History = GenerateSuggestions();
 
-					History.Add(FConsoleSuggestion(Name, HelpString));
-				}
+				// Set data to the suggestion widget
 				SetSuggestions(History, FText::GetEmpty());
 				
 				if(Suggestions.HasSuggestions())
@@ -578,6 +574,49 @@ FReply SConsoleInputBox::OnPreviewKeyDown(const FGeometry& MyGeometry, const FKe
 	return FReply::Unhandled();
 }
 
+TArray<FConsoleSuggestion> SConsoleInputBox::GenerateSuggestions() const
+{
+	IConsoleManager& ConsoleManager = IConsoleManager::Get();
+	TArray<FString> HistoryNames;
+
+	// See if there is an external history provider
+	bool bHistoryOverriden = false;
+	if (HistoryOverride.IsBound())
+	{
+		const FName ExecutorName = ActiveCommandExecutor ? ActiveCommandExecutor->GetName() : NAME_None;
+		bHistoryOverriden = HistoryOverride.Execute(ExecutorName, HistoryNames);
+	}
+
+	// If no external history provider, then fetch internally
+	if (!bHistoryOverriden)
+	{
+		if (ActiveCommandExecutor)
+		{
+			ActiveCommandExecutor->GetExecHistory(HistoryNames);
+		}
+		else
+		{
+			ConsoleManager.GetConsoleHistory(TEXT(""), HistoryNames);
+		}
+	}
+
+	TArray<FConsoleSuggestion> History;
+	for (const FString& Name : HistoryNames)
+	{
+		FString HelpString;
+		// Try to find a console object for this history entry in order to retrieve a help string if possible :
+		const TCHAR* NamePtr = *Name;
+		if (IConsoleObject* CObj = ConsoleManager.FindConsoleObject(*FParse::Token(NamePtr, /*UseEscape = */false), /*bTrackFrequentCalls = */false); CObj && CObj->IsEnabled())
+		{
+			HelpString = CObj->GetDetailedHelp().ToString();
+		}
+
+		History.Add(FConsoleSuggestion(Name, HelpString));
+	}
+
+	return History;
+}
+
 void SConsoleInputBox::SetSuggestions(TArray<FConsoleSuggestion>& Elements, FText Highlight)
 {
 	TOptional<FString> SelectionText;
diff --git a/Engine/Source/Developer/OutputLog/Private/SOutputLog.h b/Engine/Source/Developer/OutputLog/Private/SOutputLog.h
--- a/Engine/Source/Developer/OutputLog/Private/SOutputLog.h
+++ b/Engine/Source/Developer/OutputLog/Private/SOutputLog.h
@@ -3,20 +3,22 @@
 #pragma once
 
 #include "CoreMinimal.h"
-#include "SlateFwd.h"
-#include "Widgets/DeclarativeSyntaxSupport.h"
+
+#include "ConsoleInputBoxParameters.h"
+#include "Framework/Text/BaseTextLayoutMarshaller.h"
+#include "HAL/CriticalSection.h"
+#include "HAL/IConsoleManager.h"
 #include "Input/Reply.h"
+#include "Misc/TextFilterExpressionEvaluator.h"
 #include "OutputLogCreationParams.h"
 #include "OutputLogSettings.h"
-#include "Widgets/SWidget.h"
-#include "Widgets/SCompoundWidget.h"
+#include "SlateFwd.h"
+#include "Widgets/DeclarativeSyntaxSupport.h"
 #include "Widgets/Input/SMultiLineEditableTextBox.h"
-#include "Widgets/Views/STableViewBase.h"
+#include "Widgets/SCompoundWidget.h"
+#include "Widgets/SWidget.h"
 #include "Widgets/Views/STableRow.h"
-#include "Framework/Text/BaseTextLayoutMarshaller.h"
-#include "Misc/TextFilterExpressionEvaluator.h"
-#include "HAL/CriticalSection.h"
-#include "HAL/IConsoleManager.h"
+#include "Widgets/Views/STableViewBase.h"
 
 class FMenuBuilder;
 class FOutputLogTextLayoutMarshaller;
@@ -52,6 +54,8 @@ struct FOutputLogMessage
 class SConsoleInputBox
 	: public SCompoundWidget
 {
+	using FInterceptConsoleCommandDelegate = UE::OutputLog::FInterceptConsoleCommandDelegate;
+	using FHistoryOverrideDelegate = UE::OutputLog::FHistoryOverrideDelegate;
 
 public:
 	DECLARE_DELEGATE_OneParam(FExecuteConsoleCommand, const FString& /*ExecCommand*/)
@@ -63,6 +67,12 @@ public:
 		/** Where to place the suggestion list */
 		SLATE_ARGUMENT( EMenuPlacement, SuggestionListPlacement )
 
+		/** Called before executing a command; allows intercepting it. */
+		SLATE_EVENT( FInterceptConsoleCommandDelegate, OnInterceptConsoleCommand)
+
+		/** Called to gather a command history for a command executor */
+		SLATE_EVENT( FHistoryOverrideDelegate, OnHistoryOverride)
+
 		/** Custom executor for console command, will be used when bound */
 		SLATE_EVENT( FExecuteConsoleCommand, ConsoleCommandCustomExec)
 
@@ -113,7 +123,10 @@ protected:
 	TSharedRef<ITableRow> MakeSuggestionListItemWidget(TSharedPtr<FConsoleSuggestion> Message, const TSharedRef<STableViewBase>& OwnerTable);
 
 	void SuggestionSelectionChanged(TSharedPtr<FConsoleSuggestion> NewValue, ESelectInfo::Type SelectInfo);
-		
+
+	/** Genearates suggestions for the console */
+	TArray<FConsoleSuggestion> GenerateSuggestions() const;
+
 	void SetSuggestions(TArray<FConsoleSuggestion>& Elements, FText Highlight);
 
 	void MarkActiveSuggestion();
@@ -215,6 +228,12 @@ private:
 	/** Delegate to call when a console command is executed */
 	FSimpleDelegate OnConsoleCommandExecuted;
 
+	/** Delegate used to intercept console commands. */
+	FInterceptConsoleCommandDelegate InterceptConsoleCommand;
+
+	/** Delegate used to override the command history of the specified executors. */
+	FHistoryOverrideDelegate HistoryOverride;
+
 	/** Delegate to call to execute console command */
 	FExecuteConsoleCommand ConsoleCommandCustomExec;
 
diff --git a/Engine/Source/Developer/OutputLog/Public/OutputLogModule.h b/Engine/Source/Developer/OutputLog/Public/OutputLogModule.h
--- a/Engine/Source/Developer/OutputLog/Public/OutputLogModule.h
+++ b/Engine/Source/Developer/OutputLog/Public/OutputLogModule.h
@@ -3,9 +3,11 @@
 #pragma once
 
 #include "CoreMinimal.h"
-#include "Widgets/SWidget.h"
+#include "ConsoleInputBoxParameters.h"
 #include "Misc/NamePermissionList.h"
 #include "Modules/ModuleInterface.h"
+#include "OutputLogCreationParams.h"
+#include "Widgets/SWidget.h"
 #include "Widgets/SWindow.h"
 
 class SMultiLineEditableTextBox;
@@ -52,7 +54,17 @@ public:
 	/** Generates a console input box widget.  Remember, this widget will become invalid if the
 		output log DLL is unloaded on the fly. */
 	virtual TSharedRef<SWidget> MakeConsoleInputBox(TSharedPtr<SMultiLineEditableTextBox>& OutExposedEditableTextBox, const FSimpleDelegate& OnCloseConsole, const FSimpleDelegate& OnConsoleCommandExecuted) const;
-	
+
+	/**
+	 * Creates a console input box widget
+	 *
+	 * @note The widget becomes invalid if the OutputLog module is unloaded at runtime
+	 *
+	 * @param InCreationParams Input box creation parameters
+	 * @return Result structure that contains the requested widget and additional data
+	 */
+	virtual UE::OutputLog::FConsoleInputBoxCreationResult MakeConsoleInputBox(const UE::OutputLog::FConsoleInputBoxCreationParams& InCreationParams) const;
+
 	virtual TSharedRef<SWidget> MakeOutputLogDrawerWidget(const FSimpleDelegate& OnCloseConsole);
 
 	virtual TSharedRef<SWidget> MakeOutputLogWidget(const FOutputLogCreationParams& Params);
